Copying //UE4/Dev-Niagara to //UE4/Dev-Main (Source: //UE4/Dev-Niagara @ 4074996)
============================
  MAJOR FEATURES & CHANGES
============================

Change 3853627 by [REDACTED_PERSON]

	Merging using OrionDevNiagaraToUE4DevNiagara_DoNotUse VectorVM

	#tests non-gpu auto tests pass

Change 3853628 by [REDACTED_PERSON]

	Merging using OrionDevNiagaraToUE4DevNiagara_DoNotUse Runtime

	#tests all non-gpu auto tests pass

Change 3853629 by [REDACTED_PERSON]

	Merging using OrionDevNiagaraToUE4DevNiagara_DoNotUse Engine\Shaders

	#tests all non-gpu auto tests pass

Change 3853630 by [REDACTED_PERSON]

	Merging using OrionDevNiagaraToUE4DevNiagara_DoNotUse Engine\Plugins\FX

	#tests all non-gpu auto tests pass

Change 3853631 by [REDACTED_PERSON]

	Jonathan's material function from Orion\DevNiagara

	#tests all non-gpu auto tests pass

Change 3853633 by [REDACTED_PERSON]

	Merging using OrionDevNiagaraToUE4DevNiagara_DoNotUse EngineTest

	#tests all non-gpu auto tests pass

Change 3853911 by [REDACTED_PERSON]

	GPU rendering now works

	#tests GPU tests now pass

Change 3854179 by [REDACTED_PERSON]

	Removing dead system

	#tests now just with a warning

Change 3854731 by [REDACTED_PERSON]

	Checkpointing current work

	#tests n/a

Change 3855080 by [REDACTED_PERSON]

	Fixing not all control paths return a value error

	#tests n/a

Change 3856185 by [REDACTED_PERSON]

	MultiStack with pinning support.

	#tests none

Change 3856615 by [REDACTED_PERSON]

	Preventing a null pointer dereference when copying unallocated data

	#tests auto tests pass

Change 3856622 by [REDACTED_PERSON]

	Getting rid of the bogus get alias method

	#tests auto tests pass

Change 3856644 by [REDACTED_PERSON]

	Adding the ability to query the number of filtered triangles, index those triangles directly, and compute the position, velocity, UV, and NBT of a triangle.

	#tests all auto tests now pass

Change 3856645 by [REDACTED_PERSON]

	Added several new auto tests and tweaked existing ones.
	+ PerParticleRandom still had some randomness in it
	+ UserColorCurve has a user color curve defined in 3 different components referencing the same system
	+ BasicSkinnedEmitter has 4 skinning sub-tests, attached as a subcomponent, referencing a world skeletal mesh, using a material filter, and using a bone filter

	+ added a skinning module that spawns based off exec index and a data interface that queries the number of triangles on a skeletal mesh.

	#tests auto tests pass

Change 3856675 by [REDACTED_PERSON]

	Fixing crash on delete of an emitter

	#tests removed one and multiple emitters from active multi-emitter. No crash

Change 3860613 by jonathan.lindquist

	New dynamic input

Change 3862549 by [REDACTED_PERSON]

	Missing last known good images

Change 3864525 by [REDACTED_PERSON]

	Fix for vm compiler crash when using structs as constants.

	#tests No longer crashes.

Change 3864729 by [REDACTED_PERSON]

	Sequencer - Fixed a few places which were modifying sequencer data, but not calling the NotifySequencerDataChanged.

Change 3864737 by [REDACTED_PERSON]

	Niagara - Fix the timeline in the niagara editor plus other fixes.
	+ Turned on looping in the timeline by default.
	+ Added simulation options to control playback in the editor, including turning off auto-play, disabling reset on change, and disabling resimulation when changing while paused.
	+ Added a buttons to the timeline for each renderer an emitter has which shows a renderer specific icon and will allow navigation directly to the renderer with future stack changes.
	+ Fixed issues in the emitter life cycle and spawn rate modules which were preventing delay and looping from working consistently. (includes auto-test)
	+ Added top level metadata for modules.
	+ Added the ability to add metadata to a module and it's inputs to allow it to be edited directly using timed sections in the timeline.  Currently configured for the emitter life cycle module.
	+ Changed the way the "MaxSimTime" on niagara component works so that it now represents the maximum frame time which should be used when seeking the component to the desired age.  This previously was the maximum amount of simulation time to run which would prevent simple effects from simulating quickly, and would also allow more complicated effects to hang the UI while seeking.
	+ Changed the behavior of niagara component when seeking to desired age so that it always uses the exact specified seek delta instead of trying to seek to the exact desired age.
	+ Fixed the component so that systems which are seeking to the desired age in the editor no longer draw in fast forward mode.
	+ Changed the default playback range for effects in the timeline from 1000 seconds to 10 seconds, and fixed the timeline so that the playback range is serialized into the asset which will persist the setting across loads.
	+ Fixed scrubbing in the timeline so that it doesn't immediately reset when scrubbing backwards.
	+ Added a button to the timeline track for enabling and disabling isolation mode for emitters.
	+ Added a checkbox to the timeline to enabled and disable emitters from the timeline.
	+ Fixed PinToNiagaraVariable so that it asserts if bNeedsValue is specified and it can't actually provide a value.
	+ Create a class called FNiagaraStackFunctionInputBinder which allows binding to an input of a function on a stack so that you can easily set and get the value of the input without having to worry about modifying the graph or rapid iteration parameter sets directly.

	#tests Existing auto-tests and 1 new test for life cycle changes.

Change 3867179 by [REDACTED_PERSON]

	Niagara - Turn off GPU simulation for test assets for Shaun.

Change 3869201 by [REDACTED_PERSON]

	Bypassing JonLs issue

	#tests no longer crashes.

Change 3869897 by [REDACTED_PERSON]

	Niagara - Fix crashes when using the skeletal mesh data interface.  NOTE: The change to the actual skeletal mesh data interface code doesn't seem needed, but without it, it crashes with the stomp allocator on.  We'll have to investigate further.

Change 3870487 by [REDACTED_PERSON]

	Niagara - Always generate cached skin data immediately, and make sure we're not indexing triangles that don't exist.  This is a temporary fix to avoid a crash when changing the skeletal mesh source on an effect in the level from the details panel.

Change 3877378 by [REDACTED_PERSON]

	Niagara - Update the burst and lifecycle modules with new metadata for incoming timeline changes.

Change 3877564 by [REDACTED_PERSON]

	Sequencer - Fix a few more places which were modifying sequencer data without calling NotifyMovieSceneDataChanged.

Change 3877565 by [REDACTED_PERSON]

	Niagara - Remove old unused burst code from some runtime classes.

Change 3877567 by [REDACTED_PERSON]

	Niagara - Add support for keying bursts on the timeline which is configured using script metadata.

Change 3877699 by [REDACTED_PERSON]

	Niagara - Fix a crash in the new timeline code for when you have bursts, but you have inputs on the emitter lifecycle bound, also set lower bound of view range of the timeline to be -.1 so that you can more easily interact with keys at 0.

Change 3877715 by [REDACTED_PERSON]

	Sequencer - Update the change type when pasting keys from the clipboard from Unknown to TrackValueChanged to avoid unnecessary work.  From code review on last change.

Change 3879285 by [REDACTED_PERSON]

	A couple of fixes for using struct constants

	#tests Jon's case now compiles and works correctly.

Change 3879378 by [REDACTED_PERSON]

	Niagara - Fix a few spots where recursive graph traversal was visiting nodes multiple times because of diamonds in the pin connections.  This reduces a terrible hang in UNiagaraScriptSource::InitializeNewRapidIterationParameters from 5 minutes to 5 seconds but there are futher issues to investigate becasue even 5 seconds it too long.

Change 3879858 by [REDACTED_PERSON]

	Moved the VM script compilation to the DDC in order to facilitate better team compilation behavior. Significant changes to the translator and overall compile workflow to make the data behave better for a future multi-threaded compilation path.
	In order to know when to compile, a key is generated that uses the compile id's of the graphs that influence the compilation of our node. In this way, if an end user goes and edits a function or module and checks in, the overall compile id will not match b/c that function is in the dependency list of a system downstream. The system will then know to recompile. However, everyone else on the team will generate the same key because the asset in question was checked in with its local compile id already changed.

	Additionally, we now employ change tracking on traversals from a node graph. These traversals walk through all the nodes leading to a given output node and see if they've been altered. If they have, a new compile id is generated. If not, the old compile id is used. This also means that if you edit a particle update section in the stack, the emitter section won't force the system to recompile.

	GPU scripts now have their own script slot rather than riding alongside particle spawn scripts. This allows us to address them independently in the translator and put them in the DDC as well. Once the text is generated, we then go back out to the DDC to generate the shader associated.

	Known issues:
	+ Emitters are sometimes marked dirty on open
	+ Nodes connected to event writes aren't part of the hashing
	+ DataInterface signature changes don't dirty the compiles
	+ Struct changes don't dirty the compiles
	+ On system loading, we go out to the DDC instead of using existing scripts, which is slower..

	#tests all auto-tests pass, additional tests run to validate proper behavior

Change 3879859 by [REDACTED_PERSON]

	Content update post DDC change

Change 3879862 by [REDACTED_PERSON]

	Niagara plugin content to ddc

Change 3879958 by [REDACTED_PERSON]

	Niagara - Actually fix the bad recursion in this function by using the existing traversal method.

Change 3881727 by [REDACTED_PERSON]

	Niagara - Houdini - Created a separate plug-in for the Houdini CSV Data Interface.

Change 3881877 by [REDACTED_PERSON]

	Fix for mac compile issue

Change 3882773 by [REDACTED_PERSON]

	Actual fix for Mike

Change 3882822 by [REDACTED_PERSON]

	Rather than throw a check, I instead emit an error when we can't match up a data interface and instantiate a CDO version. Not perfect, but this will let you recompile.

	#tests allows me to open jonathan's file

Change 3883538 by [REDACTED_PERSON]

	Moving particle-level scripts to compile with the emitter named Emitter in their internal scripts. This simplifies the dependencies quite a bit, but causes some complexity on the wiring side (most of which we were already doing anyway).
	Getting rid of some allocations in translation (still more to go).
	Fixed some of the logic for emitters that had modules of the same name to now properly concatenate.

	Compile version bumped, so all scripts will be forced to recompile. Not saving this into the test files for now, as I expect this to happen a bit for the near term.

	#tests all auto-tests pass, creating a new emitter and system on PC works

Change 3883552 by [REDACTED_PERSON]

	Fixing renaming to work properly now. It just invalidates the system script compile id's, forcing it to auto-compile.

	#tests auto-tests pass

Change 3884722 by [REDACTED_PERSON]

	Added searchbar with basic name search for Niagara stack

	#tests none

Change 3884793 by [REDACTED_PERSON]

	Adding pragma once

	#tests no longer complains about duplicate definition

Change 3885629 by [REDACTED_PERSON]

	Setting up a transient meshrotation framework pre-integration

Change 3887440 by [REDACTED_PERSON]

	Custom HLSL failure for Shaun

Change 3888911 by [REDACTED_PERSON]

	stack search box now has a minimum width of 300 pixels

	#tests none

Change 3890843 by [REDACTED_PERSON]

	Creating a Niagara quaternion type.

	#tests created in editor, saw default was correct, carried through to VM runtime through attribute viewer

Change 3890849 by [REDACTED_PERSON]

	Porting over 4.19 fix to Dev-Niagara

	#tests allows creation of valid scripts even when ini is cleared.

Change 3891088 by [REDACTED_PERSON]

	Sequencer - When getting selected tracks for the external selection api, include tracks if any of their child nodes are selected.  This matches the behavior object guid external selection.

Change 3891114 by [REDACTED_PERSON]

	Fixing crash that sometimes happens if a stack tree changes while a stack search is active

	#tests none

Change 3891131 by [REDACTED_PERSON]

	Sequencer - Move section headers for bool, int, vector, and color to the public directory so they can be used by the niagara level sequence integration.

Change 3891165 by [REDACTED_PERSON]

	error for shaun

Change 3891354 by [REDACTED_PERSON]

	Adding Quat struct to more locations. Now treated like hlsl float4.

	#tests EulerToQuaternion now compiles

Change 3891463 by [REDACTED_PERSON]

	Fix crash that sometimes happens when deleting module and hitting Ctrl-Z to undo (the condition for removing the listeners from the rootentry should not be tied to the validity of weak pointers for system and emitter, because sometimes they are out of sync when changing the graph)

	#tests none

Change 3891641 by [REDACTED_PERSON]

	resave node

Change 3893143 by [REDACTED_PERSON]

	Fixing issue where you try to bind a vertex color sampler to a mesh without it. We failthe binding rather than crash due to a check later. Also fixed up error logging to only mention the one that failed.

	#tests can open NiagaraSystem'/Game/FX/SkeletalMeshDissolve/EmittersAndSystems/FightSceneDissolve.FightSceneDissolve' with an error..

Change 3893528 by [REDACTED_PERSON]

	fix another crash when search results are invalidated during search

	#tests none

Change 3893830 by [REDACTED_PERSON]

	Fix for copy & paste of comment boxes

	#tests can now copy and paste comment boxes

Change 3894012 by [REDACTED_PERSON]

	no longer executing search tick if the rootentry is null

	#tests none

Change 3894828 by [REDACTED_PERSON]

	Niagara - Runtime changes to support sequencer animation
	+ Reset the simulation when force solo is changed on the component.
	+ Invalidate the render data and clear the buffers when resetting to avoid previously rendered particles from drawing.
	+ Automatically sync the override parameters in the component when the source assets exposed parameters change and removed forced syncing from various places.
	+ Remove lots of refresh code from the niagara component details which should not be neccessary anymore.

	#TESTS Ran autotests, tested through the UI while building the sequencer tests asests.

Change 3894832 by [REDACTED_PERSON]

	Niagara - Level sequence support for spawning and animating system life cycle and select user parameter types.

	#Tests Ran existing tests and added a new test to verify added functionality.

Change 3896944 by [REDACTED_PERSON]

	safeguard entries with no search items (it actually can be null)

	#tests none

Change 3896948 by [REDACTED_PERSON]

	Fix assert when dereferencing source array (no need for a raw pointer to the array since it's initialized with the content anyway)

	#tests none

Change 3896950 by [REDACTED_PERSON]

	fix compile error with previous commit

	#tests none

Change 3897698 by [REDACTED_PERSON]

	Niagara - Fix some safety issues with parameter initialization on the niagara component.
	+ Kill the current system if we synchronize parameters to avoid issues with data interface lifetime.
	+ Always sync parameters when the asset is set to prevent missing data interfaces in the override list.
	+ Add an enum to control how data interface parameters are handled when calling CopyParametersTo.
	+ When the system instance is copying the asset parameters, have it copy data interfaces by reference so that it's not creating data interfaces copies which will be deleted at the next garbage collection.

	#Tests Auto-tests, also doesn't crash anymore when opeining Jonathan's disolve effect and then opening a level.

Change 3897953 by [REDACTED_PERSON]

	Niagara - Remove some namespace restrictions from the code that generates the available parameters to read from in the stack UI since it was preventing the use of custom namespaces in modules and it was not clear what issue it was solving since we already prevent scripts from addressing parameters from lower level scripts.  Also move user parameters to their own menu section.

	#Tests Custom namespaces are usable again in the stack.

Change 3898926 by [REDACTED_PERSON]

	Fix for crash caused by garbage collection and async search
	Both are still under code review, submitting because they are simple to rollback and harmless to other features.

	#tests none

Change 3899069 by [REDACTED_PERSON]

	Parallel compilation

	Major changes:
	Rather than a custom streaming version that we know influences a rebuild, I'm moving away to a guid that you need to regenerate if you change the compiler in any meaningful way
	needed for multiple reasons,
	1) if two people are making changes to the compiler, having something other than a guid as the value makes the content of the ddc ambiguous
	2) when iterating I often need to make multiple changes to get to a working final result, bumping the version number each time that happens gets old fast

	We fully clone the input graph to do the compile in the background. While the translation step is not a huge amount of time, it keeps the main thread responsive.

	We currently have a big critsec around the crosscompiler to bytecode as it isn't threadsafe. Future changes will push this to the ShaderCompilerWorker.

	#tests all tests pass as well as stress tests around saving while compilng, long compile times, etc.

Change 3899071 by [REDACTED_PERSON]

	Fixing the availability flags for system and emitter scripts.

	#tests all auto tests pass

Change 3899077 by [REDACTED_PERSON]

	Fixing assets to have their wait on compile finished checkbox checked in the editor for testing

	#tests n/a

Change 3899114 by [REDACTED_PERSON]

	Random bool custom hlsl node

Change 3899184 by [REDACTED_PERSON]

	implemented categories for module inputs (now inputs can be assigned a category in the module editor and they will be grouped by those categories in the Niagara stack)

	#tests none

Change 3899329 by [REDACTED_PERSON]

	fix broken commit by adding missing new files NiagaraStackInputCategory.cpp and .h

	#tests none

Change 3899439 by [REDACTED_PERSON]

	Niagara reroute node.

Change 3899516 by [REDACTED_PERSON]

	Official angle conversion modules.

	#tests made a local test emitter that converted back and forth between angles. Results were correct.

Change 3900193 by [REDACTED_PERSON]

	Fixing build

	#tests now compiles

Change 3900474 by [REDACTED_PERSON]

	Fixes to help Mac compile

	#tests n/a

Change 3901131 by [REDACTED_PERSON]

	Warmup feature.

	CPU Sim only.

	Also has ability to advance simulation by tick count or seconds via BP/C++.

	Includes some engine tests.

	#tests editor + autotests

Change 3901455 by [REDACTED_PERSON]

	Niagara - Add WITH_EDITORONLY_DATA to prevent non-editor compile failures.

Change 3902477 by [REDACTED_PERSON]

	Niagara - Fix FNiagaraEditorTypeUtilities to be a thread safe TSharedFromThis since it's always created with a thread safe shared pointer, also fix up issues related to this change.  Fixes a crash which occurrs when it's the target object of a delegate binding.

	#Tests adding a curve data interface to a parameter collection no longer crashes.

Change 3903478 by [REDACTED_PERSON]

	No longer doing the check if compiling on load is enabled as this always forces different change ids'

	#tests n/a

Change 3903783 by [REDACTED_PERSON]

	Trimming down excess log spew

	#tests auto-tests pass

Change 3905753 by [REDACTED_PERSON]

	Made Sine(Degrees), Sine(Radians), and Sine, and the variants thereof for trig functions.

	#tests n/a

Change 3905759 by [REDACTED_PERSON]

	Auto tests for mesh orientation

	#tests these now pass

Change 3905762 by [REDACTED_PERSON]

	These files needed to be resaved for some reason to keep passing.

Change 3906727 by [REDACTED_PERSON]

	Curve UX improvements

	#tests none

Change 3908177 by [REDACTED_PERSON]

	Fixing build due to typo

	#tests now compiles

Change 3908199 by [REDACTED_PERSON]

	Trying to fix compilation when destroying objects. We cannot safely attach anything beneath us at this point ,we just need to clear out the queues.

	#tests normal work day-to-day

Change 3908201 by [REDACTED_PERSON]

	Working to fix crashes where the component was destroyed out from underneath us due to PIE shutting down and we have a Niagara item editable in Blueprint or world editor.

	#tests n/a

Change 3908985 by [REDACTED_PERSON]

	Renaming ColorCurveAsset to CurveAsset to better reflect the actual usage of the variable (fixing copy-paste issue)

	#tests none

Change 3909222 by [REDACTED_PERSON]

	Niagara graph connection colors

Change 3909436 by [REDACTED_PERSON]

	fix crash in curve ux when importing a linear curve (curve of floats)

	#tests none

Change 3909561 by [REDACTED_PERSON]

	Updating LUT before sending NotifyPostChange when editing curves inline (so that LUT will not go out of sync)

	#tests none

Change 3910010 by [REDACTED_PERSON]

	Use new Niagara Actor icon

Change 3910191 by [REDACTED_PERSON]

	Fix viewport widget showing up in the viewport when pressing W, E or R.

Change 3910213 by [REDACTED_PERSON]

	PropertyEditor - PropertyRowGenerator - Added features and fixes to support integration into niagara's stack view.
	+ Added a method to get filter/search strings for an IDetailTreeNode to support external searching and filtering.
	+ Added a delegate to the layout builder for when one of it's owned nodes has it's visibility forcibly changed by a customization.
	+ Changed the filtering so nodes are generated for properties marked as advanced.
	+ Pass the notify hook down to the detail utilities so that change notifications work as expected.
	+ Add layout data for the widgets returned from the IDetailTreeNode to prevent alignment and sizing issues in custom implementations.

Change 3910307 by [REDACTED_PERSON]

	PropertyEditor - Missed in last checkin.

Change 3910509 by [REDACTED_PERSON]

	Niagara - Removed nested details panels from the stack and integrate them properly plus other fixes.
	+ Generate rows for nested objects using the details panel property row generator.
	+ Fix the horizontal sizing for niagara parameter editors.
	+ Add an IsValid() method to the base niagara stack entry so that derived classes can know if the associated view models are still valid when processing events.  This is a temporary measure to fix a crash in the user parameter UI.
	+ Set stack entries to be expandable by default and delete usages which were setting it to true.
	+ Highlight the active search result with a border since property rows can't highlight text.

Change 3911653 by [REDACTED_PERSON]

	Niagara - Fix stack spacer sizing.

Change 3911667 by [REDACTED_PERSON]

	PropertyEditor - Actually fix the notify hook handling in the property row generator.

Change 3911896 by [REDACTED_PERSON]

	Niagara function input context menu.

Change 3911900 by [REDACTED_PERSON]

	Project setting for not showing comment bubbles.

Change 3911996 by [REDACTED_PERSON]

	Niagara fix if node persistent guids for older nodes. The OutputVarGuids are always synced on PostLoad.

Change 3912221 by [REDACTED_PERSON]

	Renderer Icons for timeline

Change 3912608 by [REDACTED_PERSON]

	stack style refactor

	#tests none

Change 3913063 by [REDACTED_PERSON]

	Icons for stack added, including new system param png

Change 3913618 by [REDACTED_PERSON]

	Fixing two of the most common Illegal call to StaticFindObject() errors while compiling.

	#tests ran through compilation after changes.

Change 3914369 by [REDACTED_PERSON]

	Using new SystemParams.png icon provided by Wyett (instead of the old "Parameters.png")

	#tests none

Change 3914782 by [REDACTED_PERSON]

	Adjusting icon for update to not indicate "flow"

Change 3915738 by [REDACTED_PERSON]

	Moving away from the generic and super-slow EdGraphSchema ShouldAlwaysPurgeOnModification being true to using the same mechanism we use to invalidate the compile to synchronize nodes. This should be substantially faster.

	#tests ran through a variety of tests creating and wiring nodes

Change 3915739 by [REDACTED_PERSON]

	Assignment nodes need to invalidate the graph for compile.

Change 3915741 by [REDACTED_PERSON]

	Making default values more accessible and making it possible to route renderers to use different values than the defaults.

	#tests n/a

Change 3915798 by [REDACTED_PERSON]

	SearchBox - Add options to show the number of search results and an option to show a throbber when a search is active.

Change 3915966 by [REDACTED_PERSON]

	Changing the default for velocity to 0,0,0 as requested  by Wyeth

	#tests n/a

Change 3915982 by [REDACTED_PERSON]

	Making the default text more readable

	#tests n/a

Change 3916237 by [REDACTED_PERSON]

	PropertyEditor - Change the DetailCategoryBuilderImpl so that it sets the horizontal alignment to fill for value widgets when generating stand alone widgets so that the behavior in the property row generator matches the behavior of the property grid.

Change 3916240 by [REDACTED_PERSON]

	Niagara - Should prevent some recent crashes due to stack entry delegates and lifetime.

Change 3916261 by [REDACTED_PERSON]

	Niagara - Lots of minor stack ui fixes and adjustments
	+ Tweaked padding in a bunch of different places.
	+ Added a dark background behind the stack and stack header to prevent the colors from bleeding together.
	+ Fixed the group text not being white anymore.
	+ Hooked up new features of the search box for showing the search result data and an is searching throbber.
	+ Fixed an issue where the current search result couldn't be interacted with.
	+ Fix some other inconsistencies with searching where you might jump more than one result.
	+ Replace the checkbox for showing curve in the curves tab with an icon based button. (icon is placeholder)

Change 3916833 by [REDACTED_PERSON]

	Fixing issue where the system wasn't set to wait for compilation on load, sometimes leading to failures for auto-tests

	#tests this test now passes when forced to recompile

Change 3916846 by [REDACTED_PERSON]

	Missed one system in the scene.

	#tests n/a

Change 3917458 by [REDACTED_PERSON]

	Fixing another potential race condition on the DDC.

	#tests n/a

Change 3918349 by [REDACTED_PERSON]

	Niagara - Invalidate the node visuals when reallocating pins.

Change 3918783 by [REDACTED_PERSON]

	Correct 'temp' to 'Temp' in map set

Change 3919262 by [REDACTED_PERSON]

	We weren't properly updating the default values for user data interface components when tweaked in the editor.

	#tests open skinned mesh auto test system change the preview for the user skinned mesh to be SK_Mannequin_Niagara. It now updates, it didn't before.

Change 3919602 by [REDACTED_PERSON]

	Fixing the skeletal mesh to now clamp to the end of the index buffer for safety as well as adding IsValidTriCood. This lets us keep going even when swapping out the skeletal mesh underneath.
	Tested out isvalidtricoord in the test skinning module.

	#tests auto tests pass
	#codereivew simon.tovey

Change 3921701 by [REDACTED_PERSON]

	Make Vector2 and Vector4 default blue color to be consistent with blueprints.

Change 3922331 by [REDACTED_PERSON]

	Niagara - Houdini - Added support of CSV File as UAsset (HoudiniCSV)
	Modified the Data Interface to use the CSV asset instead of the imported buffers from the CSV File Path.
	Added some new functions to the DI:
	GetLastParticleIndexAtTime()
	GetCSVPositionAndTime()
	GetCSVVectorValue()
	GetCSVFloatValueByString()

Change 3923118 by [REDACTED_PERSON]

	PS4 compile fix.

Change 3924934 by [REDACTED_PERSON]

	fix Mac compile issues

	#tests none

Change 3925168 by [REDACTED_PERSON]

	Curve logspamming

	#tests none
	The UpdateCompiledDataInterfaces would end up comparing LUTs when copying curves and the source LUT was out of date.

Change 3925366 by [REDACTED_PERSON]

	Slate - SMenuAnchor - Fix the implementation of "BelowRightAnchor" to align the right edge of the menu with the right edge of the anchor.  There aren't any other usages of this in the engine as far as I can tell, hopefully people weren't relying on the broken behavior in a game somwhere.

Change 3925423 by [REDACTED_PERSON]

	Niagara - Remove the large add buttons from the stack and add smaller add buttons in the group headers.

Change 3925877 by [REDACTED_PERSON]

	New collision modules, separating query, linear and angular impulse; Solve forces and velocity takes care of integrating f->v->p and fA->vA->O; linear impulse module would probably be cleaner by zeroing velocity on collision and calculating a force instead of setting new velocity directly

Change 3926582 by [REDACTED_PERSON]

	PS4 compile fix

Change 3927401 by [REDACTED_PERSON]

	Fixing events due to added member

	#tests all tests pass as of 3925423 with this change

Change 3927496 by [REDACTED_PERSON]

	Getting auto-tests to run
	Questions: Why did I have to recompile the GPU tests... something is missing in their key generation?
	Resaved several files.

	#tests almost all pass now

Change 3927582 by [REDACTED_PERSON]

	Fixing last failing auto test

	#tests all tests now pass

Change 3927924 by [REDACTED_PERSON]

	Chunk level vm parallelism.
	Any execution processing > batch_size chunks will go wide.

	The batch size is 4 currently but adjsutable via vm.ParallelChunksPerBatch.
	VM parallelism can be disabled by vm.Parallel 0

Change 3927990 by [REDACTED_PERSON]

	Submitting redirector

Change 3928426 by [REDACTED_PERSON]

	Niagara - Always propagate rapid iterations parameters when merging an emitter.

Change 3929823 by [REDACTED_PERSON]

	Niagara - Fix hlsl generation for system/emitter spawn script so that we read the engine and user parameters from the data set instead of initializing them to 0.

	#Tests Full recompile + auto-tests

Change 3929983 by [REDACTED_PERSON]

	Curve LUT Interpolation

	+ updated test altered by it.

Change 3930551 by [REDACTED_PERSON]

	Niagara - Fix what looks like a copy/paste error in the SNiagaraSelectedEmitterGraph destructor which was preventing clean removal of delegates and causing a crash.

	#Tests closing the "Selected Emitter Graph" tab and then changing the selected emitter no longer crashes.

Change 3932695 by [REDACTED_PERSON]

	Niagara - Houdini:

	Houdini CSV Asset:
	- Packed vector values in the CSV file are now properly supported (not just for Position/Normal) and can be of any size.
	- Added support for reimporting Houdini CSV files.
	- Added an "open in text editor" entry in the context menu.
	- Improved error/warning logging during the parsing of the file

	Houdini Niagara Data Interface:
	- Added GetParticleIndexesToSpawnAtTime():
	New helper functions returning the min index, max index and number of particles to be spawned for a given time value.
	Uses an internal LastSpawnIndex to avoid spawning the same particles twice.
	- Modified GetLastParticleIndexAtTime():
	If the CSV file doesn't have time informations, returns false and set the LastIndex to the last particle
	If desiredTime is smaller than the first particle, LastIndex will be set to -1
	If desiredTime is higher than the last particle in the csv file, LastIndex will be set to the last particle's index

Change 3933425 by [REDACTED_PERSON]

	Made the spreadsheet debugger capable of capturing in-world systems as long as they are solo'ed.

	#tests have been running with it for several days, debugging real-world assets stably

Change 3933986 by [REDACTED_PERSON]

	Niagara - Fixed a bug with merging where added dynamic inputs which changed names could end up with the wrong rapid iteration parameters.  Also fixed an issue where added dynamnic inputs would be renamed when they didn't need to be.

	#Tests Engine tests and fixes custom repro.

Change 3934052 by [REDACTED_PERSON]

	Niagara - Added a console command to dump rapid iteration parameters for a system or emitter asset.

Change 3934436 by [REDACTED_PERSON]

	Fixes for sprite VF depth test failure issue

Change 3934658 by [REDACTED_PERSON]

	Niagara - Make disabled modules visually distinct.

	#Tests General stack use.

Change 3935383 by [REDACTED_PERSON]

	Fixing mac compile errors

	#tests n/a

Change 3935420 by [REDACTED_PERSON]

	Niagara parameter UI first version.

Change 3935482 by [REDACTED_PERSON]

	Add missing files for parameters

Change 3935591 by [REDACTED_PERSON]

	more macos compile

	#tests na

Change 3935637 by [REDACTED_PERSON]

	Reverting to prior behavior

	#tests na

Change 3936541 by [REDACTED_PERSON]

	Remove the merge up menu entry for set variables module items.

Change 3936841 by [REDACTED_PERSON]

	Bool comparison dynamic input

Change 3936895 by [REDACTED_PERSON]

	A few perf improvements and fixes to the SetSolo transfering between solo and batched so all lightning sims can run batched after they're warmded up.

Change 3936899 by [REDACTED_PERSON]

	Missed a file

Change 3937178 by [REDACTED_PERSON]

	Fixed bHasSkipOutputVelocityParameter for shaders without PreviousLocalToWorldMatrix (e.g. particles)

Change 3937222 by [REDACTED_PERSON]

	Random event spawn

Change 3937292 by [REDACTED_PERSON]

	Fix Adding a new parameter then renaming it the default name deletes the new parameter

Change 3938472 by [REDACTED_PERSON]

	Fix new parameters in emitters saving by using the editable emitter.

Change 3938474 by [REDACTED_PERSON]

	- Store graphs as weak object pointers in the parameter UI.
	- Allow right mouse menu on parameters in the system toolkit.
	- Refresh only the parameter actions when deleting an entry instead of refreshing the graphs aswell.

Change 3938525 by [REDACTED_PERSON]

	Fix creating an unique FName every tick for parameterstores by using a FString instead.

Change 3938596 by [REDACTED_PERSON]

	Macos compile

	#tests n/a

Change 3939362 by jonathan.lindquist

	Adding a new Component Spacing input to the debug value functions. This will allow users to make better use of space when debugging values.

Change 3939365 by [REDACTED_PERSON]

	Back out changelist 3936895 and 3936899
	Leaving in some changes around stats as they should be harmless.
	These changes were removed b/c they added poor perf to Jonathan's dissolve effect and also caused multiple tests to fail in engine tests.

	#tests all tests pass besides Yannick's FName/FString warning, with the exception of BPTimeControl, Hypnotizer, and MeshOrientation, which seem to be off by one frame, but have been consistently off for several days (CL 3929823 had same issues for me)

Change 3939367 by jonathan.lindquist

	Adding greyscale output

Change 3939368 by jonathan.lindquist

	Changing the pin order

Change 3939377 by [REDACTED_PERSON]

	Allows the unnormalized lut table flag to be copied over

	#tests all tests pass besides Yannick's FName/FString warning, with the exception of BPTimeControl, Hypnotizer, and MeshOrientation, which seem to be off by one frame, but have been consistently off for several days (CL 3929823 had same issues for me)

Change 3939379 by [REDACTED_PERSON]

	Rename FParameterStore Name to DebugName to prevent loading a FName into a FString with existing assets.

Change 3939382 by [REDACTED_PERSON]

	Adding the ability to have a default curve index with a custom switch node.

	#tests all tests pass besides Yannick's FName/FString warning, with the exception of BPTimeControl, Hypnotizer, and MeshOrientation, which seem to be off by one frame, but have been consistently off for several days (CL 3929823 had same issues for me)

Change 3939383 by [REDACTED_PERSON]

	Converting existing curves over to using the new default pin

	#tests all tests pass besides Yannick's FName/FString warning, with the exception of BPTimeControl, Hypnotizer, and MeshOrientation, which seem to be off by one frame, but have been consistently off for several days (CL 3929823 had same issues for me)

Change 3939501 by [REDACTED_PERSON]

	Submitting missing files

	#tests n/a

Change 3939580 by [REDACTED_PERSON]

	Default curve indexing to three more DIs

Change 3940122 by [REDACTED_PERSON]

	Parameters view: - Jump to new parameter added and request rename for new parameter.
	- Remove adding parameters to the parameterstore when a pin is requested.
	- Only show make new parameters in the dropdown to add a new parameter.
	- Use Sections of UI as types instead of int32.

Change 3940214 by [REDACTED_PERSON]

	fix crash when rename skeletal mesh user variable

	#tests none

Change 3940215 by [REDACTED_PERSON]

	undo not working in graph editor

	#tests none
	Overriding the BreakPinLinks methods to include a transaction

Change 3940250 by [REDACTED_PERSON]

	Creating stats tab in module toolkit to show LastOpCount

	#tests none

Change 3940251 by [REDACTED_PERSON]

	#tests none

Change 3940262 by [REDACTED_PERSON]

	Back out changelist 3939365 with fixes
	Tests all now pass

Change 3940333 by [REDACTED_PERSON]

	Nullptr check

	#tests n/a

Change 3940338 by [REDACTED_PERSON]

	Niagara sprite particles - implemented get previous position in order to fix sprite particle motion vectors

Change 3940407 by [REDACTED_PERSON]

	Create pin on map get and set node when dragging without recompiling the graph.

Change 3940534 by [REDACTED_PERSON]

	Making sure that collision returns defaults of 0 if nothing was found.

	#tests auto-tests that have been passing still pass

Change 3940709 by [REDACTED_PERSON]

	Temp hacks for the skeletal mesh painting issues.

Change 3940960 by [REDACTED_PERSON]

	Only build parameter menu once when graphchanged is called multiple times in a frame.
	Also use the existing metadata from graph to build the parameter menu, because the metadata already looped through all nodes and pins.

Change 3941019 by [REDACTED_PERSON]

	Meta data UI refresh next tick to avoid refreshing multiple times a tick.

Change 3941853 by [REDACTED_PERSON]

	Adding more dynamic parameters

Change 3941957 by [REDACTED_PERSON]

	Property Editor - Fix issues with property row generator to support the niagara stack.
	+ Make the detail tree node name accessible through the interface, and fix the implementations for category group and property item.
	+ Add a temporary fix for passing instance customizations from the property row generator to the detail property row through the detail layout builder.  This should be unified in a nicer way, but this will work for the time being.

Change 3942174 by [REDACTED_PERSON]

	Niagara - Stack UI Pass
	+ Advanced rows are not handled properly per item.
	+ Expanded and scroll state is now saved in editor data per asset.
	+ Added a "View Options" drop down for showing all advanced rows, and for showing/hiding outputs.
	+ Added an option to collapse all stack items from the emitter header context menu.
	+ Added support for "Edit Conditions" on module and dymaic inputs which will enable and disabled an input based on the value of another input.  This includes showing a checlbox inline for the edit condition toggle input.
	+ Added support for "Visible Conditions" on module and dynamic inputs which will hide and show inputs based on the value of another input.
	+ Removed the pencil icon for locally editable values in the stack.
	+ Fixed issues with invisible search results.  It's still possible for a search result to not highlight the text correctly, but the outline never disappears.
	+ Removed pinning for module inputs.
	+ Fixed the event handler properties so that they use the property row generator instead of an embedded details panel.
	+ Unified indent handling across all stack classes.
	+ Unified stack editor data across all stack classes.

Change 3942427 by [REDACTED_PERSON]

	Another hack for vertex painting tool

Change 3942453 by [REDACTED_PERSON]

	Some more hacks for skel mesh vertex painting until Jurre's rework is ready.

Change 3942799 by [REDACTED_PERSON]

	Rebuild metadata, input and output parameters UI next frame instead of on every graph changed call.

Change 3942833 by [REDACTED_PERSON]

	Niagara - Fix the visibility of the advanced expander item.

Change 3942923 by [REDACTED_PERSON]

	Revert using metadata for parameters to looping through pins to find parameters in maps not connected to anything.
	Temporary fix to remove used parameters if they are found in the graph.

Change 3943094 by [REDACTED_PERSON]

	Rollback //UE4/Dev-Niagara/Engine/Plugins/FX/Niagara/Content/Modules/Spawn/Location/SphereLocation.uasset to revision 5

Change 3943154 by [REDACTED_PERSON]

	Metadata to sphere location module

Change 3943256 by [REDACTED_PERSON]

	Testing out sweet new metadata control

Change 3943374 by [REDACTED_PERSON]

	Fixing mesh motion blur

Change 3943382 by [REDACTED_PERSON]

	Turning on base pass velocities until I can fix separate vel pass for particles

Change 3943471 by [REDACTED_PERSON]

	Emitter view stats only show particles count. Use Niagara.EmitterStatsFormat 0 and 1 to switch between all data and only particle count. Default is 1 to only show particle count.

Change 3943497 by [REDACTED_PERSON]

	Paramater map remove FTickableEditorObject and use SWidget::Tick

Change 3943589 by [REDACTED_PERSON]

	-Fix for linear impulse (offset by 1/2*velocity*dt instead to avoid distracting bounciness)
	-Collision Rest; add after impulse modules to make particles rest if in collision under threshold velocity magnitude

Change 3943644 by [REDACTED_PERSON]

	Turn shader development mode back off

Change 3943718 by [REDACTED_PERSON]

	Fix vertex factories

Change 3943776 by [REDACTED_PERSON]

	Properly calculate old particle position using dt for velocity rendering

Change 3943780 by [REDACTED_PERSON]

	Niagara - Fix ensure when removing dynamic inputs due to incorrect logic which would have left unused nodes in the graph.

Change 3943870 by [REDACTED_PERSON]

	Parameter drag drop window

Change 3943994 by [REDACTED_PERSON]

	Niagara - Fix some editor settings not saving across sessions.

Change 3944056 by [REDACTED_PERSON]

	Updating to replace values reset when Wyeth resaved files.

	#tests DrawOrderGPU, DrawOrder, and TestDifferentInactiveEmitters now pass

Change 3944068 by [REDACTED_PERSON]

	Back out of my changes to vertex painting and replaced with Jurre's fixes.

Change 3944174 by [REDACTED_PERSON]

	Niagara - Fix stack categories so that they don't freak out when the stack is refreshed, and also remove the uncategorized heading and move uncategorized inputs to the top.

Change 3944313 by [REDACTED_PERSON]

	Updated screenshots after motion blur

Change 3944321 by [REDACTED_PERSON]

	Fixing error message to be clearer about a disconnected Get node.

	#tests n/a

Change 3944351 by [REDACTED_PERSON]

	Making safe against weak pointers going away.
	Track error encountered by Jonathan.

	#tests n/a

Change 3944368 by [REDACTED_PERSON]

	Remove automatic adding prefix "particles."  when renaming functioninput.

Change 3944383 by [REDACTED_PERSON]

	Just adding some more nullptr check

	#tests n/a

Change 3944384 by [REDACTED_PERSON]

	Providing more context for a check that existed previosuly  and was encountered by Wyeth today

	#tests n/a

Change 3944872 by [REDACTED_PERSON]

	Remove old material parameter node.

Change 3945209 by [REDACTED_PERSON]

	Fixing possible infinite recurson on child array size of zero

	#tests n/a

Change 3945865 by [REDACTED_PERSON]

	Spreadsheet filter for output attributes

Change 3946091 by [REDACTED_PERSON]

	Per particle sorting for translucent sprites and meshes

Change 3946095 by [REDACTED_PERSON]

	Updated screens for dynamic param tests

Change 3946378 by [REDACTED_PERSON]

	Another sprite motion blur fix

Change 3946864 by [REDACTED_PERSON]

	SkinnedMesh per-instance data requires 16 byte alignment due to usage of FMatrix. We were not guaranteeing that in our per-instance data system. We are now enforcing that to be true by aligning all memory size requests.

	#tests autotests pass

Change 3946928 by [REDACTED_PERSON]

	Skeletal mesh location metadata. THIS VERSION ALSO CRASHES ON SAVE FYI

Change 3946934 by [REDACTED_PERSON]

	Niagara - Clean up rapid iteration parameters on compile.

	#Tests - Fixes the jira below, all auto tests which were currently passing still pass, and GDC effects load and look correct.

Change 3946936 by [REDACTED_PERSON]

	Niagara - Fix crash when undoing adding a dynami input.

Change 3947213 by [REDACTED_PERSON]

	Fix for thread safety in collision data interface.
	Previously I'd made them thread safe between VM chunks but they already weren't safe between system instances.

Change 3947279 by [REDACTED_PERSON]

	Fixed thread safety issue with niagara global dynamic buffer

Change 3947788 by [REDACTED_PERSON]

	Fix enum property warnings

Change 3947849 by [REDACTED_PERSON]

	Normalize orientation quats. Safety first.

Change 3947877 by [REDACTED_PERSON]

	Niagara - Fix a crash when editing meta-data for a module currently open in a system or emitter editor stack.  This updates the FNiagaraStackFunctionInputBinder to track the lifetime of the default pin correctly.

Change 3948445 by jonathan.lindquist

	Inverting alpha output

Change 3948615 by [REDACTED_PERSON]

	Don't access data layouts that are invalid because their bindings don't exist

Change 3949361 by [REDACTED_PERSON]

	Command to expand all groups and collapse all items of those groups in the stack.

Change 3949365 by [REDACTED_PERSON]

	Missing file for change 3949361

Change 3951123 by [REDACTED_PERSON]

	Fix bug with dynamic parameters in Niagara mesh particle VFs

Change 3951199 by [REDACTED_PERSON]

	Fix for issues caused by unsafe reads of GT data from RT

Change 3951293 by [REDACTED_PERSON]

	Workaround for jittering particles with collision at rest state; will need to revisit after GDC

Change 3951533 by [REDACTED_PERSON]

	Collapse parameter menu by default

Change 3952106 by [REDACTED_PERSON]

	Niagara - Fix data interface input initialization when inserting modules and dynamic inputs.  We now put all inputs into categories in the stack and this code didn't handle that.

Change 3954809 by [REDACTED_PERSON]

	HoudiniNiagara - Add include to fix CIS incremental build.

Change 3954857 by [REDACTED_PERSON]

	Niagara - Accept newer versions of 3 automated tests images as they are stable and still look correct for what they are testing.

Change 3954935 by [REDACTED_PERSON]

	Niagara - Fix a crash in the skeletal mesh sampling info details customization when a mesh has no skeleton.

Change 3954969 by [REDACTED_PERSON]

	Compile fix for gpu emitters

Change 3955012 by [REDACTED_PERSON]

	Niagara - Fix clang and deprecation warnings.

Change 3955988 by [REDACTED_PERSON]

	Fixing collision queries (separating line query trace direction and velocity, so we can look a frame ahead properly); various fixes to the impulse modules; rest module now allows for color change when particles are set to rest. This checkin should stabilize collision dynamics substantially.

Change 3956730 by [REDACTED_PERSON]

	Cleanup parameter

Change 3957065 by [REDACTED_PERSON]

	enable/disable mechanism for renderers

	#tests none

Change 3957802 by [REDACTED_PERSON]

	-Removing Velocity parameter from collision query DI, since velocity at query time is known and can just be passed along in temporary parameter; this stops the compiler from falling over in unity builds and also makes the Perform Query function more sane as an actual line check
	-Put some safeguarding against non-collisions (v.n>0) reported as intersections into the collision query module to avoid instances of particles being pushed through geometry
	-Updated modules to use the new function signature

	#tests modular explosion test map, EngineTests

Change 3957804 by [REDACTED_PERSON]

	Updated tests for 3957802

Change 3957859 by [REDACTED_PERSON]

	Niagara - Add missing #if to fix some of the nightly build errors.

Change 3958065 by [REDACTED_PERSON]

	Fix GPU sim hlsl for the collision data interface; should make GPU collision bounce test run again

Change 3958302 by [REDACTED_PERSON]

	modified test for sane depth bounds; accepting new results; some changes to hlsl for collision data interface

Change 3959007 by [REDACTED_PERSON]

	Further defining the barrier between GT and RT data.

	Not finished yet, especially for GPU sims but we're heading in the right direction.

Change 3960004 by [REDACTED_PERSON]

	QOL change, now committing search text in the stack (pressing enter) will jump to the next occurrence.

	#tests none

Change 3960019 by [REDACTED_PERSON]

	Niagara - Preemptively fix up stack related header includes to avoid manual merges.

Change 3964217 by [REDACTED_PERSON]

	Fixing compile issue after renderer enable feature(the Isolate features require the WITH_EDITORONLY_DATA preprocessor directive to be enabled)

	#tests none

Change 3964581 by [REDACTED_PERSON]

	Niagara - Get things compiling again after merge.  Mostly include fixes and commeting out lots of sequencer related stuff that needs to be fixed.

Change 3965057 by [REDACTED_PERSON]

	Niagara - Fix compile issues in the houdini plugin.  Also add the houdini pluging to EngineTest so that it compiles by default when running tests.

Change 3965075 by [REDACTED_PERSON]

	Niagara - Fix another include issue that was caught on the build machine.

Change 3965308 by [REDACTED_PERSON]

	Niagara - One more header fix.  Should fix the win64 build in CIS.

Change 3965313 by [REDACTED_PERSON]

	Niagara - Fix in editor playback.  The timeline data is still broken.

Change 3965482 by [REDACTED_PERSON]

	Stack source taken apart into different files.

Change 3965863 by [REDACTED_PERSON]

	Fixes scope level variable definition causing my local build to fail

	#tests n/a

Change 3965866 by [REDACTED_PERSON]

	Crash fix when the module is missing. Now show an error message as well as checking for script validity before calling method on it

	#tests n/a

Change 3968174 by [REDACTED_PERSON]

	Niagara - Fix more merge fallout.  The emitter/system editor timeline now matches the stack data again.

Change 3968183 by [REDACTED_PERSON]

	Niagara - Delete commented out include from merge.

Change 3972162 by [REDACTED_PERSON]

	Niagara - Updated level sequence testing assets.

Change 3972880 by [REDACTED_PERSON]

	Merging using DevNiagaraToGDC2018
	Allowing disabled modules to still influence parameter maps.

	#tests n/a

Change 3973269 by [REDACTED_PERSON]

	Disabling warning about divide by zero as it is often incorrect

	#tests n/a

Change 3973273 by [REDACTED_PERSON]

	Forcing all three planes to be GPU

	#tests n/a

Change 3973307 by [REDACTED_PERSON]

	Fixing CIS win32 errors

	#tests n/a

Change 3973374 by [REDACTED_PERSON]

	Fixing minor static analysis warnings

	#tests n/a

Change 3976107 by [REDACTED_PERSON]

	Updating multiple files as they have the unversioned file warning

	#tests auto tests now show green

Change 3976114 by [REDACTED_PERSON]

	Taking snapshots after the integration for time control to clear automated tests. Frank was uncertain about the current behavior being correct, but didn't want to hold up integration for that.

Change 3976119 by [REDACTED_PERSON]

	Makiing GPU shaders contain the dependencies and the compile id's and other items so that they update properly.

	#tests auto-tests now pass

Change 3976449 by [REDACTED_PERSON]

	Adding additional debugging to logs

	#tests n/a

Change 3977172 by [REDACTED_PERSON]

	Niagara - Fix issues with the level sequence integration for niagara caused by the integration from main, and accept the new test image.  This code should be updated at some point to use the new channel blending in sequencer.

	NOTE: There is still a timing issue that is evident in the automated test screen shot due to float timing and rounding issues in niagara, but the sequencer code is functioning correctly.

Change 3977362 by [REDACTED_PERSON]

	UE-55601 curve snapping not working in Niagara curve editor

	#tests none

Change 3977363 by [REDACTED_PERSON]

	exposed added external asset for function input so it would appear in context menu in the Niagara stack.

	#tests none

Change 3977368 by [REDACTED_PERSON]

	The CachedUsageInfo got emptied but it didn't get restored by the undo, sending the system in an endless Compile() loop
	Added a Modify() call in the UNiagaraScriptSource::InvalidateCachedCompileIds()

	#tests none

Change 3978716 by [REDACTED_PERSON]

	Fixing half of CIS static analysis warnings and localization symbol dupe warnings from Jamie Dale

	#tests auto tests pass other than known level sequence test

Change 3978857 by [REDACTED_PERSON]

	The map SpawnTest niagara actor didn't have its wait for compilation flag set, potentially leading to artifacts. Disabling for now as that didn't resolve the issue.

Change 3979594 by [REDACTED_PERSON]

	Potential fix for cook on Orion

	#tests n/a

Change 3979713 by [REDACTED_PERSON]

	Fixing several more CIS static analysis warnings as well as duplicate localization string warnings from Jamie Dale

	#tests n/a

Change 3980017 by [REDACTED_PERSON]

	Fixing CIS static analysis warnings as well as duplicate localization string warnings from Jamie Dale

	#tests n/a

Change 3981859 by [REDACTED_PERSON]

	Fixing crash in Paragon when the normals buffer was not present during startup with just UI screens.

	#tests got through to paragon main screen, related auto-tests pass

Change 3982685 by [REDACTED_PERSON]

	For some reason, when Lightmass is spinning off workers, the render thread is null, causing us to trigger checks that we shouldn't trigger when killing of system instances.

	#tests all auto-tests pass

Change 3983902 by [REDACTED_PERSON]

	Speculative fixes for mac errors regarding niagara vertex decls.

Change 3984023 by [REDACTED_PERSON]

	Sequencer: No longer upgrade bIsInfinite for section types that do not support open ranges

Change 3986727 by [REDACTED_PERSON]

	Fixed add velocity from point to work in more situations, not cause errors regardless of stack location, made it more programmable with inputs, and added metadata

Change 3988114 by [REDACTED_PERSON]

	fixed color and inherit parent vel

Change 3989175 by [REDACTED_PERSON]

	Improved VM error reporting.

	#tests engine tests

Change 3995007 by [REDACTED_PERSON]

	Parameter menu in system layout

Change 3995192 by [REDACTED_PERSON]

	Fix niagara script details panel search

Change 3995291 by [REDACTED_PERSON]

	Parameter menu tooltip fix

Change 3997804 by shaun.kime

	Lookup table is off for this. Keys are 0,0 and 1,1, so you'd expect normalized age as a sampler to just return the same value from lookup. If LUT is disabled, this is true. If LUT is enabled, it isn't, especially the farther one gets from 0.

Change 3998124 by [REDACTED_PERSON]

	Fixed Curve LUT generation

	#tests editor

Change 3998286 by [REDACTED_PERSON]

	Fixed normalized execution index to, you know, work.

Change 4000324 by [REDACTED_PERSON]

	Replacing a thread-safety issue where we get an enum on an as-needed basis. This was causing crashes in cooking as you cannot find from the global table while serializing.

	#tests have not gotten the crash since

Change 4000428 by [REDACTED_PERSON]

	UE-55750 focus curve editor when "show" button is pressed
	UE-55791 user variable curves cannot be shown in the Curve tab
	These issues were related and touched the same area of code, so I fixed them together. UE-55791 is basically just getting curve data from the system exposed variables too.

	#tests none

Change 4001094 by [REDACTED_PERSON]

	Niagara - Fix slowdown in the metadata editor due to delegate rebinding.

Change 4001098 by [REDACTED_PERSON]

	Niagara - Remove the tool tip from the additional options drop-down for modules since it covers the drop down menu.

Change 4001133 by [REDACTED_PERSON]

	Fix curve editor getting focused each time a curve was changed.

	#tests none

Change 4001253 by [REDACTED_PERSON]

	PropertyEditor - Fix issues with external root handling.
	+ Fix external nodes not being cleaned up correctly when custom node builders rebuild children.
	+ Fix expanded state being trampled when processing external nodes due to the expanded nodes list being emptied every time RestoreExpandedItems was called.
	+ Fix performance issues with refreshing during tick by moving all calls to RestoreExpandedItems to UpdateFilteredDetails, and then only calling UpdateFilteredDetails once per tick as needed instead of per root property node and per root external node.

Change 4003365 by [REDACTED_PERSON]

	If the physical material has gone away, make sure to set value values. Note that if this data interface is going to go on in parallel to the game thread, we'll need something more sophisticated.

	#tests Win64 tests pass

Change 4003367 by [REDACTED_PERSON]

	Making sure that the system has finished compiling before we begin cooking.

	#tests n/a

Change 4003374 by [REDACTED_PERSON]

	Niagara - Fix a crash when adding and removing modules and dynamic inputs with data interfaces.

Change 4003696 by [REDACTED_PERSON]

	Getting rid of whitelist, now open on all platforms.

	#tests n/a

Change 4005368 by [REDACTED_PERSON]

	Fixing compile error on Linux

	#tests n/a

Change 4013779 by [REDACTED_PERSON]

	Interpolated spawn on GPU does not yet work but checkpointing work.
	+ Added AdditionalDefines to the VMCompileId, switched interpoalted spawn to use that
	+ Added rough interpolated spawn support to translator by refactoring away from specific calls and hard-coded update/spawn somewhat... could improve for events
	+ Revised shader variables away from Phase0/Phase1 terminology to the true Update/Spawn meaning. Leaving the phase numbers as an implementation detail within the usf.
	+ Added ToString to ParameterStores for debugging
	+ Changed GPU ExecIndex() logic to actually work similar to the VM

	#tests collision gpu fails

Change 4015355 by [REDACTED_PERSON]

	Persistent IDs final.

	Still need to change the compile ID parts over to use the new additional defines but the core functionality is in.

	#tests editor + engine tests

Change 4018445 by [REDACTED_PERSON]

	Some missing assets

Change 4021647 by shaun.kime

	Moving jonathan's DebugParticleData to Niagara Extras

Change 4024809 by [REDACTED_PERSON]

	Parameter map hover text using metadata description.

Change 4025042 by [REDACTED_PERSON]

	Dogfooding the skeletal Mesh location module w/ comments, reroute pins, map gets and sets, metadata, tooltips, etc.

Change 4025236 by [REDACTED_PERSON]

	Working on getting interpolated spawning working. Submitting to get assistance from Simon.

	... PLEASE DON'T SYNC UNLESS YOU'VE SPOKEN TO ME...

	#tests collision test on GPU is better, but not yet right.. ribbon id emitter is ensuring on LUT table generation

Change 4025372 by [REDACTED_PERSON]

	Making the debug particle data stay the same size as the base particle for easier debugging.

	#tests n/a

Change 4025701 by [REDACTED_PERSON]

	Debug asset for Simon

	#tests n/a

Change 4027865 by [REDACTED_PERSON]

	Fixing parameter map stores to properly handle reset. Previously was leaving around padding info from previous version of the script which could be totally wrong.

	#tests now don't crash

Change 4029638 by [REDACTED_PERSON]

	Refactor Skeletal mesh location module to test some ideas on coding standards and shake out workflow issues

Change 4030135 by [REDACTED_PERSON]

	Interpolated spawning now works on the GPU.

	#tests collision gpu has stray collisions and the GenerateLocationEvent was recently updated that broke . will fix in a later update

Change 4030197 by [REDACTED_PERSON]


diff --git a/Engine/Config/BaseEditorPerProjectUserSettings.ini b/Engine/Config/BaseEditorPerProjectUserSettings.ini
--- a/Engine/Config/BaseEditorPerProjectUserSettings.ini
+++ b/Engine/Config/BaseEditorPerProjectUserSettings.ini
@@ -766,7 +766,7 @@ bAutoScrollEnabled=true
 bKeepPlayRangeInSectionBounds=false
 bKeepCursorInPlayRange=false
 bShowRangeSlider=true
-bLooping=true
+LoopMode=SLM_Loop
 
 
 [/Script/LevelSequenceEditor.LevelSequenceEditorSettings]
diff --git a/Engine/Config/BaseScalability.ini b/Engine/Config/BaseScalability.ini
--- a/Engine/Config/BaseScalability.ini
+++ b/Engine/Config/BaseScalability.ini
@@ -326,7 +326,7 @@ r.Streaming.AmortizeCPUToGPUCopy=1
 r.Streaming.MaxNumTexturesToStreamPerFrame=2
 r.MaxAnisotropy=0
 r.Streaming.LimitPoolSizeToVRAM=1
-r.Streaming.PoolSize=400
+r.Streaming.PoolSize=1000
 r.Streaming.MaxEffectiveScreenSize=0
 
 [TextureQuality@1]
@@ -335,7 +335,7 @@ r.Streaming.AmortizeCPUToGPUCopy=0
 r.Streaming.MaxNumTexturesToStreamPerFrame=0
 r.MaxAnisotropy=2
 r.Streaming.LimitPoolSizeToVRAM=1
-r.Streaming.PoolSize=600
+r.Streaming.PoolSize=2000
 r.Streaming.MaxEffectiveScreenSize=0
 
 [TextureQuality@2]
@@ -344,7 +344,7 @@ r.Streaming.AmortizeCPUToGPUCopy=0
 r.Streaming.MaxNumTexturesToStreamPerFrame=0
 r.MaxAnisotropy=4
 r.Streaming.LimitPoolSizeToVRAM=1
-r.Streaming.PoolSize=800
+r.Streaming.PoolSize=3000
 r.Streaming.MaxEffectiveScreenSize=0
 
 [TextureQuality@3]
@@ -353,7 +353,7 @@ r.Streaming.AmortizeCPUToGPUCopy=0
 r.Streaming.MaxNumTexturesToStreamPerFrame=0
 r.MaxAnisotropy=8
 r.Streaming.LimitPoolSizeToVRAM=0
-r.Streaming.PoolSize=1000
+r.Streaming.PoolSize=3000
 r.Streaming.MaxEffectiveScreenSize=0
 
 [TextureQuality@Cine]
diff --git a/Engine/Config/ConsoleVariables.ini b/Engine/Config/ConsoleVariables.ini
--- a/Engine/Config/ConsoleVariables.ini
+++ b/Engine/Config/ConsoleVariables.ini
@@ -31,8 +31,10 @@
 
 [Startup]
 
+
 ; Uncomment to get detailed logs on shader compiles and the opportunity to retry on errors
-;r.ShaderDevelopmentMode=1
+; r.ShaderDevelopmentMode=1
+r.BasePassOutputsVelocity=1
 ; Uncomment to dump shaders in the Saved folder
 ; Warning: leaving this on for a while will fill your hard drive with many small files and folders
 ;r.DumpShaderDebugInfo=1
diff --git a/Engine/Plugins/FX/HoudiniNiagara/HoudiniNiagara.uplugin b/Engine/Plugins/FX/HoudiniNiagara/HoudiniNiagara.uplugin
new file mode 100644
--- /dev/null
+++ b/Engine/Plugins/FX/HoudiniNiagara/HoudiniNiagara.uplugin
@@ -0,0 +1,39 @@
+{
+	"FileVersion": 3,
+	"Version": 1,
+	"VersionName": "1.0",
+	"FriendlyName": "Houdini Niagara",
+	"Description": "Houdini Niagara Integration",
+	"Category": "FX",
+	"CreatedBy": "Side Effects Software Inc.",
+	"CreatedByURL": "http://www.sidefx.com",
+	"DocsURL": "",
+	"MarketplaceURL": "",
+	"SupportURL": "",
+	"CanContainContent": true,
+	"IsBetaVersion": true,
+	"Installed": false,
+	"Modules":
+        [
+		{
+			"Name": "HoudiniNiagara",
+			"Type": "Runtime",
+			"LoadingPhase": "Default",
+                        "WhitelistPlatforms" : [ "Win64", "Win32", "PS4", "XboxOne", "Mac" ]
+		},
+                {
+			"Name": "HoudiniNiagaraEditor",
+			"Type": "Editor",
+			"LoadingPhase": "Default",
+                        "WhitelistPlatforms" : [ "Win64", "Win32", "PS4", "XboxOne", "Mac" ]
+		}
+	],
+        
+        "Plugins":
+        [
+                {
+                        "Name": "Niagara",
+                        "Enabled": true
+                }
+        ]  
+}
\ No newline at end of file
diff --git a/Engine/Plugins/FX/HoudiniNiagara/Source/HoudiniNiagara/Classes/NiagaraDataInterfaceHoudiniCSV.h b/Engine/Plugins/FX/HoudiniNiagara/Source/HoudiniNiagara/Classes/NiagaraDataInterfaceHoudiniCSV.h
new file mode 100644
--- /dev/null
+++ b/Engine/Plugins/FX/HoudiniNiagara/Source/HoudiniNiagara/Classes/NiagaraDataInterfaceHoudiniCSV.h
@@ -0,0 +1,126 @@
+/*
+* Copyright (c) <2017> Side Effects Software Inc.
+*
+* Permission is hereby granted, free of charge, to any person obtaining a copy
+* of this software and associated documentation files (the "Software"), to deal
+* in the Software without restriction, including without limitation the rights
+* to use, copy, modify, merge, publish, distribute, sublicense, and/or sell
+* copies of the Software, and to permit persons to whom the Software is
+* furnished to do so, subject to the following conditions:
+*
+* The above copyright notice and this permission notice shall be included in all
+* copies or substantial portions of the Software.
+*
+* THE SOFTWARE IS PROVIDED "AS IS", WITHOUT WARRANTY OF ANY KIND, EXPRESS OR
+* IMPLIED, INCLUDING BUT NOT LIMITED TO THE WARRANTIES OF MERCHANTABILITY,
+* FITNESS FOR A PARTICULAR PURPOSE AND NONINFRINGEMENT. IN NO EVENT SHALL THE
+* AUTHORS OR COPYRIGHT HOLDERS BE LIABLE FOR ANY CLAIM, DAMAGES OR OTHER
+* LIABILITY, WHETHER IN AN ACTION OF CONTRACT, TORT OR OTHERWISE, ARISING FROM,
+* OUT OF OR IN CONNECTION WITH THE SOFTWARE OR THE USE OR OTHER DEALINGS IN THE
+* SOFTWARE.
+*
+*/
+#pragma once
+
+#include "CoreMinimal.h"
+#include "UObject/ObjectMacros.h"
+#include "NiagaraCommon.h"
+#include "NiagaraShared.h"
+#include "VectorVM.h"
+#include "HoudiniCSV.h"
+#include "NiagaraDataInterface.h"
+#include "NiagaraDataInterfaceHoudiniCSV.generated.h"
+
+
+/** Data Interface allowing sampling of Houdini CSV files. */
+UCLASS(EditInlineNew, Category = "Houdini Niagara", meta = (DisplayName = "Houdini CSV File"))
+class HOUDININIAGARA_API UNiagaraDataInterfaceHoudiniCSV : public UNiagaraDataInterface
+{
+	GENERATED_UCLASS_BODY()
+public:
+
+	//UPROPERTY( EditAnywhere, Category = "Houdini Niagara", meta = (DisplayName = "Houdini CSV File Path" ) )
+	//FFilePath CSVFileName;
+
+	UPROPERTY( EditAnywhere, Category = "Houdini Niagara", meta = (DisplayName = "Houdini CSV File" ) )
+	UHoudiniCSV* CSVFile;
+
+	//void UpdateDataFromCSVFile();
+
+	//----------------------------------------------------------------------------
+	//UObject Interface
+	virtual void PostInitProperties() override;
+	virtual void PostLoad() override;
+#if WITH_EDITOR
+	virtual void PostEditChangeProperty(struct FPropertyChangedEvent& PropertyChangedEvent) override;
+#endif
+	//----------------------------------------------------------------------------
+
+	virtual void GetFunctions(TArray<FNiagaraFunctionSignature>& OutFunctions)override;
+	
+	/** Returns the delegate for the passed function signature. */
+	virtual void GetVMExternalFunction(const FVMExternalFunctionBindingInfo& BindingInfo, void* InstanceData, FVMExternalFunction &OutFunc) override;
+
+	virtual bool Equals(const UNiagaraDataInterface* Other) const override;
+
+	//----------------------------------------------------------------------------
+	// EXPOSED FUNCTIONS
+
+	// Returns the float value at a given point in the CSV file
+	template<typename RowParamType, typename ColParamType>
+	void GetCSVFloatValue(FVectorVMContext& Context);
+
+	template<typename RowParamType, typename ColTitleParamType>
+	void GetCSVFloatValueByString(FVectorVMContext& Context);
+
+	// Returns a Vector3 value for a given point in the csv file
+	template<typename RowParamType, typename ColParamType>
+	void GetCSVVectorValue(FVectorVMContext& Context);
+
+	// Returns the positions for a given point in the CSV file
+	template<typename NParamType>
+	void GetCSVPosition(FVectorVMContext& Context);
+
+	// Returns the normals for a given point in the CSV file
+	template<typename NParamType>
+	void GetCSVNormal(FVectorVMContext& Context);
+
+	// Returns the time for a given point in the CSV file
+	template<typename NParamType>
+	void GetCSVTime(FVectorVMContext& Context);
+
+	// Returns the position and time for a given point in the CSV file
+	template<typename NParamType>
+	void GetCSVPositionAndTime(FVectorVMContext& Context);
+
+	// Returns the number of points found in the CSV file
+	void GetNumberOfPointsInCSV(FVectorVMContext& Context);
+
+	// Returns the last index of the particles that should be spawned at time t
+	template<typename TimeParamType>
+	void GetLastParticleIndexAtTime(FVectorVMContext& Context);
+
+	// Returns the indexes (min, max) and number of particles that should be spawned at time t
+	template<typename TimeParamType>
+	void GetParticleIndexesToSpawnAtTime(FVectorVMContext& Context);
+	
+	//----------------------------------------------------------------------------
+	// GPU / HLSL Functions
+	virtual bool GetFunctionHLSL(const FName& DefinitionFunctionName, FString InstanceFunctionName, FNiagaraDataInterfaceGPUParamInfo& ParamInfo, FString& OutHLSL) override;
+	virtual void GetParameterDefinitionHLSL(FNiagaraDataInterfaceGPUParamInfo& ParamInfo, FString& OutHLSL) override;
+
+	// Disabling GPU sim for now.
+	virtual bool CanExecuteOnTarget(ENiagaraSimTarget Target)const override { return Target == ENiagaraSimTarget::CPUSim; }
+
+protected:
+
+	virtual bool CopyToInternal(UNiagaraDataInterface* Destination) const override;
+
+	// Indicates the GPU buffers need to be updated
+	UPROPERTY()
+	bool GPUBufferDirty;
+
+	// Last Index used to spawn particles
+	UPROPERTY()
+	int32 LastSpawnIndex;
+};
diff --git a/Engine/Plugins/FX/HoudiniNiagara/Source/HoudiniNiagara/HoudiniNiagara.Build.cs b/Engine/Plugins/FX/HoudiniNiagara/Source/HoudiniNiagara/HoudiniNiagara.Build.cs
new file mode 100644
--- /dev/null
+++ b/Engine/Plugins/FX/HoudiniNiagara/Source/HoudiniNiagara/HoudiniNiagara.Build.cs
@@ -0,0 +1,70 @@
+/*
+* Copyright (c) <2017> Side Effects Software Inc.
+*
+* Permission is hereby granted, free of charge, to any person obtaining a copy
+* of this software and associated documentation files (the "Software"), to deal
+* in the Software without restriction, including without limitation the rights
+* to use, copy, modify, merge, publish, distribute, sublicense, and/or sell
+* copies of the Software, and to permit persons to whom the Software is
+* furnished to do so, subject to the following conditions:
+*
+* The above copyright notice and this permission notice shall be included in all
+* copies or substantial portions of the Software.
+*
+* THE SOFTWARE IS PROVIDED "AS IS", WITHOUT WARRANTY OF ANY KIND, EXPRESS OR
+* IMPLIED, INCLUDING BUT NOT LIMITED TO THE WARRANTIES OF MERCHANTABILITY,
+* FITNESS FOR A PARTICULAR PURPOSE AND NONINFRINGEMENT. IN NO EVENT SHALL THE
+* AUTHORS OR COPYRIGHT HOLDERS BE LIABLE FOR ANY CLAIM, DAMAGES OR OTHER
+* LIABILITY, WHETHER IN AN ACTION OF CONTRACT, TORT OR OTHERWISE, ARISING FROM,
+* OUT OF OR IN CONNECTION WITH THE SOFTWARE OR THE USE OR OTHER DEALINGS IN THE
+* SOFTWARE.
+*
+*/
+
+using UnrealBuildTool;
+
+public class HoudiniNiagara : ModuleRules
+{
+	public HoudiniNiagara(ReadOnlyTargetRules Target) : base(Target)
+	{
+        PCHUsage = ModuleRules.PCHUsageMode.UseExplicitOrSharedPCHs;
+
+        PrivateIncludePaths.AddRange(new string[] {
+				"HoudiniNiagara/Private",
+		});
+		PrivateDependencyModuleNames.AddRange(
+			new string[]
+			{
+				"CoreUObject",
+				"Engine",
+				"Slate",
+				"SlateCore",
+              			"NiagaraCore",
+				"Niagara",
+				"NiagaraShader",
+				"RenderCore",
+				"UtilityShaders",
+				"ShaderCore"
+			}
+		);
+			
+		
+		PublicDependencyModuleNames.AddRange(
+			new string[]
+			{
+				"Core",
+              			"NiagaraCore",
+				"Niagara",
+				"NiagaraShader",
+				"CoreUObject",
+				"VectorVM",
+				"RHI",
+				"UtilityShaders",
+				"NiagaraVertexFactories",
+				"ShaderCore"
+			}
+		);
+			
+		
+	}
+}
diff --git a/Engine/Plugins/FX/HoudiniNiagara/Source/HoudiniNiagara/Private/HoudiniCSV.cpp b/Engine/Plugins/FX/HoudiniNiagara/Source/HoudiniNiagara/Private/HoudiniCSV.cpp
new file mode 100644
--- /dev/null
+++ b/Engine/Plugins/FX/HoudiniNiagara/Source/HoudiniNiagara/Private/HoudiniCSV.cpp
@@ -0,0 +1,497 @@
+/*
+* Copyright (c) <2017> Side Effects Software Inc.
+*
+* Permission is hereby granted, free of charge, to any person obtaining a copy
+* of this software and associated documentation files (the "Software"), to deal
+* in the Software without restriction, including without limitation the rights
+* to use, copy, modify, merge, publish, distribute, sublicense, and/or sell
+* copies of the Software, and to permit persons to whom the Software is
+* furnished to do so, subject to the following conditions:
+*
+* The above copyright notice and this permission notice shall be included in all
+* copies or substantial portions of the Software.
+*
+* THE SOFTWARE IS PROVIDED "AS IS", WITHOUT WARRANTY OF ANY KIND, EXPRESS OR
+* IMPLIED, INCLUDING BUT NOT LIMITED TO THE WARRANTIES OF MERCHANTABILITY,
+* FITNESS FOR A PARTICULAR PURPOSE AND NONINFRINGEMENT. IN NO EVENT SHALL THE
+* AUTHORS OR COPYRIGHT HOLDERS BE LIABLE FOR ANY CLAIM, DAMAGES OR OTHER
+* LIABILITY, WHETHER IN AN ACTION OF CONTRACT, TORT OR OTHERWISE, ARISING FROM,
+* OUT OF OR IN CONNECTION WITH THE SOFTWARE OR THE USE OR OTHER DEALINGS IN THE
+* SOFTWARE.
+*
+*/
+
+#include "HoudiniCSV.h"
+#include "Misc/FileHelper.h"
+#include "Misc/Paths.h"
+
+#if WITH_EDITOR
+#include "EditorFramework/AssetImportData.h"
+#endif
+
+#define LOCTEXT_NAMESPACE HOUDINI_NIAGARA_LOCTEXT_NAMESPACE 
+
+DEFINE_LOG_CATEGORY(LogHoudiniNiagara);
+ 
+UHoudiniCSV::UHoudiniCSV( const FObjectInitializer& ObjectInitializer )
+	: Super( ObjectInitializer ),
+	NumberOfLines( -1 ),
+	NumberOfColumns( -1 )
+{
+        
+}
+
+void UHoudiniCSV::SetFileName( const FString& TheFileName )
+{
+    FileName = TheFileName;
+}
+
+bool UHoudiniCSV::UpdateFromFile( const FString& TheFileName )
+{
+    if ( TheFileName.IsEmpty() )
+	return false;
+
+    // Check the CSV file exists
+    const FString FullCSVFilename = FPaths::ConvertRelativePathToFull( TheFileName );
+    if ( !FPaths::FileExists( FullCSVFilename ) )
+	return false;
+
+    FileName = TheFileName;
+
+    // Parse the file to a string array
+    TArray<FString> StringArray;
+    if ( !FFileHelper::LoadFileToStringArray( StringArray, *FullCSVFilename ) )
+	return false;
+
+    return UpdateFromStringArray( StringArray );
+}
+
+bool UHoudiniCSV::UpdateFromStringArray( TArray<FString>& RawStringArray )
+{
+    // Reset the CSV sizes
+    NumberOfColumns = 0;
+    NumberOfLines = 0;
+
+    // Reset the position, normal and time indexes
+    PositionColumnIndex = INDEX_NONE;
+    NormalColumnIndex = INDEX_NONE;
+    TimeColumnIndex = INDEX_NONE;
+
+    if ( RawStringArray.Num() <= 0 )
+    {
+	UE_LOG( LogHoudiniNiagara, Error, TEXT( "Could not load the CSV file, error: not enough lines." ) );
+	return false;
+    }
+
+    // Remove empty lines from the CSV
+    RawStringArray.RemoveAll( [&]( const FString& InString ) { return InString.IsEmpty(); } );
+
+    // Number of lines in the CSV (ignoring the title line)
+    NumberOfLines = RawStringArray.Num() - 1;
+    if ( NumberOfLines < 1 )
+    {
+	UE_LOG( LogHoudiniNiagara, Error, TEXT( "Could not load the CSV file, error: not enough lines." ) );
+	return false;
+    }
+
+    // Get the number of values per lines via the title line
+    FString TitleString = RawStringArray[ 0 ];
+    TitleRowArray.Empty();
+    TitleString.ParseIntoArray( TitleRowArray, TEXT(",") );
+    NumberOfColumns = TitleRowArray.Num();
+    if ( NumberOfColumns < 1 )
+    {
+	UE_LOG( LogHoudiniNiagara, Error, TEXT( "Could not load the CSV file, error: not enough columns." ) );
+	return false;
+    }
+
+    // Look for the position, normal and time attributes indexes
+    for ( int32 n = 0; n < TitleRowArray.Num(); n++ )
+    {
+	if ( TitleRowArray[ n ].Equals( TEXT("P"), ESearchCase::IgnoreCase )
+	    || TitleRowArray[ n ].Equals( TEXT("Px"), ESearchCase::IgnoreCase ) 
+	    || TitleRowArray[ n ].Equals( TEXT("X"), ESearchCase::IgnoreCase ) )
+	{
+	    if ( PositionColumnIndex == INDEX_NONE )
+		PositionColumnIndex = n;
+	}
+	else if ( TitleRowArray[ n ].Equals( TEXT("N"), ESearchCase::IgnoreCase ) 
+	    || TitleRowArray[ n ].Equals( TEXT("Nx"), ESearchCase::IgnoreCase ) )
+	{
+	    if ( NormalColumnIndex == INDEX_NONE )
+		NormalColumnIndex = n;
+	}
+	else if ( ( TitleRowArray[ n ].Equals( TEXT("T"), ESearchCase::IgnoreCase ) )
+	    || ( TitleRowArray[ n ].Contains( TEXT("time"), ESearchCase::IgnoreCase ) ) )
+	{
+	    if ( TimeColumnIndex == INDEX_NONE )
+		TimeColumnIndex = n;
+	}
+    }
+
+    // Read the first line of the CSV file, and look for packed vectors value (X,Y,Z)
+    // We'll have to expand them in the title row to match the parsed data
+    FString FirstLine = RawStringArray[ 1 ];
+    int32 FoundPackedVectorCharIndex = 0;
+    bool HasPackedVectors = false;
+    while ( FoundPackedVectorCharIndex != INDEX_NONE )
+    {
+	// Try to find ( in the line
+	FoundPackedVectorCharIndex = FirstLine.Find( TEXT("("), ESearchCase::IgnoreCase, ESearchDir::FromStart, FoundPackedVectorCharIndex );
+	if ( FoundPackedVectorCharIndex == INDEX_NONE )
+	    break;
+	// We want to know which column this char belong to
+	int32 FoundPackedVectorColumnIndex = INDEX_NONE;
+	{
+	    // Chop the first line up to the found character
+	    FString FirstLineLeft = FirstLine.Left( FoundPackedVectorCharIndex );
+
+	    // ReplaceInLine returns the number of occurences of ",", that's what we want! 
+	    FoundPackedVectorColumnIndex = FirstLineLeft.ReplaceInline(TEXT(","), TEXT(""));
+	}
+
+	if ( !TitleRowArray.IsValidIndex( FoundPackedVectorColumnIndex ) )
+	{
+	    UE_LOG( LogHoudiniNiagara, Warning,
+		TEXT( "Error while parsing the CSV File. Couldn't unpack vector found at character %d in the first line!" ),
+		FoundPackedVectorCharIndex + 1 );
+	    continue;
+	}
+	    
+
+	// We found a packed vector, get its sizw
+	int32 FoundVectorSize = 0;
+	{
+	    // Extract the vector string
+	    int32 FoundPackedVectorEndCharIndex = FirstLine.Find( TEXT(")"), ESearchCase::IgnoreCase, ESearchDir::FromStart, FoundPackedVectorCharIndex );
+	    FString VectorString = FirstLine.Mid( FoundPackedVectorCharIndex + 1, FoundPackedVectorEndCharIndex - FoundPackedVectorCharIndex - 1 );
+
+	    // Use ReplaceInLine to count the number of , to get the vector's size!
+	    FoundVectorSize = VectorString.ReplaceInline( TEXT(","), TEXT("") ) + 1;
+	}
+
+	if ( FoundVectorSize < 2 )
+	    continue;
+
+	// Increment the number of columns
+	NumberOfColumns += ( FoundVectorSize - 1 );
+
+	// Expand TitleRowArray
+	if ( ( FoundPackedVectorColumnIndex == PositionColumnIndex ) && ( FoundVectorSize == 3 ) )
+	{
+	    // Expand P to PX,Py,Pz
+	    TitleRowArray[ PositionColumnIndex ] = TEXT("Px");
+	    TitleRowArray.Insert( TEXT( "Py" ), PositionColumnIndex + 1 );
+	    TitleRowArray.Insert( TEXT( "Pz" ), PositionColumnIndex + 2 );
+	}
+	else if ( ( FoundPackedVectorColumnIndex == NormalColumnIndex ) && ( FoundVectorSize == 3 ) )
+	{
+	    // Expand N to Nx,Ny,Nz
+	    TitleRowArray[ NormalColumnIndex ] = TEXT("Nx");
+	    TitleRowArray.Insert( TEXT("Ny"), NormalColumnIndex + 1 );
+	    TitleRowArray.Insert( TEXT("Nz"), NormalColumnIndex + 2 );
+	}
+	else
+	{
+	    // Expand the vector's title
+	    FString FoundPackedVectortTitle = TitleRowArray[ FoundPackedVectorColumnIndex ];
+	    for ( int32 n = 1; n < FoundVectorSize; n++ )
+	    {
+		FString CurrentTitle = FoundPackedVectortTitle + FString::FromInt( n );
+		TitleRowArray.Insert( CurrentTitle, FoundPackedVectorColumnIndex + n );
+	    }
+	}
+
+	// Eventually offset the stored index
+	if ( PositionColumnIndex != INDEX_NONE && ( PositionColumnIndex > FoundPackedVectorColumnIndex) )
+	    PositionColumnIndex += FoundVectorSize - 1;
+
+	if ( NormalColumnIndex != INDEX_NONE && ( NormalColumnIndex > FoundPackedVectorColumnIndex) )
+	    NormalColumnIndex += FoundVectorSize - 1;
+
+	if ( TimeColumnIndex != INDEX_NONE && ( TimeColumnIndex > FoundPackedVectorColumnIndex) )
+	    TimeColumnIndex += FoundVectorSize - 1;
+
+	HasPackedVectors = true;
+	FoundPackedVectorCharIndex++;
+    }
+
+    // For sanity, Check that the number of columns matches the title row and the first line
+    {
+	// Check the title row
+	if ( NumberOfColumns != TitleRowArray.Num() )
+	    UE_LOG( LogHoudiniNiagara, Error,
+		TEXT( "Error while parsing the CSV File. Found %d columns but the Title string has %d values! Some values will have an offset!" ),
+		NumberOfColumns, TitleRowArray.Num() );
+
+	// Use ReplaceInLine to count the number of columns in the first line and make sure it's correct
+	int32 FirstLineColumns = FirstLine.ReplaceInline( TEXT(","), TEXT("") ) + 1;
+	if ( NumberOfColumns != FirstLineColumns )
+	    UE_LOG( LogHoudiniNiagara, Error,
+		TEXT("Error while parsing the CSV File. Found %d columns but found %d values in the first line! Some values will have an offset!" ),
+		NumberOfColumns, FirstLineColumns );
+    }
+    
+    // Remove the title row now that it's been processed
+    RawStringArray.RemoveAt( 0 );
+
+    // Initialize our different buffers
+    FloatCSVData.Empty();
+    FloatCSVData.SetNumZeroed( NumberOfLines * NumberOfColumns );
+    
+    StringCSVData.Empty();
+    StringCSVData.SetNumZeroed( NumberOfLines * NumberOfColumns );
+
+     // Extract all the values from the table to the float & string buffers
+    TArray<FString> CurrentParsedLine;
+    for ( int rowIdx = 0; rowIdx < NumberOfLines; rowIdx++ )
+    {
+	FString CurrentLine = RawStringArray[ rowIdx ];
+	if ( HasPackedVectors )
+	{
+	    // Clean up the packing characters: ()" from the line so it can be parsed properly
+	    CurrentLine.ReplaceInline( TEXT("("), TEXT("") );
+	    CurrentLine.ReplaceInline( TEXT(")"), TEXT("") );
+	    CurrentLine.ReplaceInline( TEXT("\""), TEXT("") );
+	}
+
+	// Parse the current line to an array
+	CurrentParsedLine.Empty();
+	CurrentLine.ParseIntoArray( CurrentParsedLine, TEXT(",") );
+
+	// Check the parsed line and number of columns match
+	if ( NumberOfColumns != CurrentParsedLine.Num() )
+	    UE_LOG( LogHoudiniNiagara, Warning,
+		TEXT("Error while parsing the CSV File. Line %d has %d values instead of the expected %d!"),
+		rowIdx + 1, CurrentParsedLine.Num(), NumberOfColumns );
+
+	// Store the CSV Data in the buffer
+	// The data is stored transposed in those buffers
+	for ( int colIdx = 0; colIdx < NumberOfColumns; colIdx++ )
+	{
+	    // Get the string value for the current column
+	    FString CurrentVal = TEXT("0");
+	    if ( CurrentParsedLine.IsValidIndex( colIdx ) )
+		CurrentVal = CurrentParsedLine[ colIdx ];
+	    else
+	    {
+		UE_LOG( LogHoudiniNiagara, Warning,
+		    TEXT("Error while parsing the CSV File. Line %d has an invalid value for column %d!"),
+		    rowIdx + 1, colIdx + 1 );
+	    }
+
+	    // Convert the string value to a float in the buffer
+	    FloatCSVData[ rowIdx + ( colIdx * NumberOfLines ) ] = FCString::Atof( *CurrentVal );
+
+	    // Keep the original string value in a buffer too
+	    StringCSVData[ rowIdx + ( colIdx * NumberOfLines ) ] = CurrentVal;
+	}	
+    }
+
+    return true;
+}
+
+
+// Returns the float value at a given point in the CSV file
+bool UHoudiniCSV::GetCSVFloatValue( const int32& lineIndex, const int32& colIndex, float& value )
+{
+    if ( lineIndex < 0 || lineIndex >= NumberOfLines )
+	return false;
+
+    if ( colIndex < 0 || colIndex >= NumberOfColumns )
+	return false;
+
+    int32 Index = lineIndex + ( colIndex * NumberOfLines );
+    if ( FloatCSVData.IsValidIndex( Index ) )
+    {
+	value = FloatCSVData[ Index ];
+	return true;
+    }
+
+    return false;
+}
+
+// Returns the float value at a given point in the CSV file
+bool UHoudiniCSV::GetCSVStringValue( const int32& lineIndex, const int32& colIndex, FString& value )
+{
+    if ( lineIndex < 0 || lineIndex >= NumberOfLines )
+	return false;
+
+    if ( colIndex < 0 || colIndex >= NumberOfColumns )
+	return false;
+
+    int32 Index = lineIndex + ( colIndex * NumberOfLines );
+    if ( StringCSVData.IsValidIndex( Index ) )
+    {
+	value = StringCSVData[ Index ];
+	return true;
+    }
+
+    return false;
+}
+
+// Returns a Vector 3 for a given point in the CSV file
+bool UHoudiniCSV::GetCSVVectorValue( const int32& lineIndex, const int32& colIndex, FVector& value, const bool& DoSwap, const bool& DoScale )
+{
+    FVector V = FVector::ZeroVector;
+    if ( !GetCSVFloatValue( lineIndex, colIndex, V.X ) )
+	return false;
+
+    if ( !GetCSVFloatValue( lineIndex, colIndex + 1, V.Y ) )
+	return false;
+
+    if ( !GetCSVFloatValue( lineIndex, colIndex + 2, V.Z ) )
+	return false;
+
+    if ( DoScale )
+	V *= 100.0f;
+
+    value = V;
+
+    if ( DoSwap )
+    {
+	value.Y = V.Z;
+	value.Z = V.Y;
+    }
+
+    return true;
+}
+
+// Returns a Vector 3 for a given point in the CSV file
+bool UHoudiniCSV::GetCSVPositionValue( const int32& lineIndex, FVector& value )
+{
+    FVector V = FVector::ZeroVector;
+    if ( !GetCSVVectorValue( lineIndex, PositionColumnIndex, V, true, true ) )
+	return false;
+
+    value = V;
+
+    return true;
+}
+
+// Returns a Vector 3 for a given point in the CSV file
+bool UHoudiniCSV::GetCSVNormalValue( const int32& lineIndex, FVector& value )
+{
+    FVector V = FVector::ZeroVector;
+    if ( !GetCSVVectorValue(lineIndex, NormalColumnIndex, V, true, false ) )
+	return false;
+
+    value = V;
+
+    return true;
+}
+
+// Returns a time value for a given point in the CSV file
+bool UHoudiniCSV::GetCSVTimeValue( const int32& lineIndex, float& value )
+{
+    float temp;
+    if ( !GetCSVFloatValue( lineIndex, TimeColumnIndex, temp ) )
+	return false;
+
+    value = temp;
+
+    return true;
+}
+
+// Returns the number of points found in the CSV file
+int32 UHoudiniCSV::GetNumberOfPointsInCSV()
+{
+    return NumberOfLines;
+}
+
+// Get the last index of the particles with a time value smaller or equal to desiredTime
+// If the CSV file doesn't have time informations, returns false and set the LastIndex to the last particle
+// If desiredTime is smaller than the first particle, LastIndex will be set to -1
+// If desiredTime is higher than the last particle in the csv file, LastIndex will be set to the last particle's index
+bool UHoudiniCSV::GetLastParticleIndexAtTime( const float& desiredTime, int32& lastIndex )
+{
+    // If we dont have proper time info, always return the last line
+    if ( TimeColumnIndex < 0 || TimeColumnIndex >= NumberOfColumns )
+    {
+	lastIndex = NumberOfLines - 1;
+	return false;
+    }
+
+    float temp_time = 0.0f;
+    // Check first if we have anything to spawn at the current time by looking at the last value
+    if ( GetCSVTimeValue( NumberOfLines - 1, temp_time ) && temp_time < desiredTime )
+    {
+	// We didn't find a suitable index because the desired time is higher than our last time value
+	lastIndex = NumberOfLines - 1;
+	return true;
+    }	
+
+    // Iterates through all the particles
+    lastIndex = INDEX_NONE;
+    for ( int32 n = 0; n < NumberOfLines; n++ )
+    {
+	if ( GetCSVTimeValue( n, temp_time ) )
+	{
+	    if ( temp_time == desiredTime )
+	    {
+		lastIndex = n;
+	    }
+	    else if ( temp_time > desiredTime )
+	    {
+		lastIndex = n - 1;
+		return true;
+	    }
+	}
+    }
+
+    // We didn't find a suitable index because the desired time is higher than our last time value
+    if ( lastIndex == INDEX_NONE )
+	lastIndex = NumberOfLines - 1;
+
+    return true;
+}
+
+// Returns the column index for a given string
+bool UHoudiniCSV::GetColumnIndexFromString( const FString& ColumnTitle, int32& ColumnIndex )
+{
+    if ( !TitleRowArray.Find( ColumnTitle, ColumnIndex ) )
+	return true;
+
+    // Handle packed positions/normals here
+    if ( ColumnTitle.Equals( "P" ) )
+	return TitleRowArray.Find( TEXT( "Px" ), ColumnIndex );
+    else if ( ColumnTitle.Equals( "N" ) )
+	return TitleRowArray.Find( TEXT( "Nx" ), ColumnIndex );
+
+    return false;
+}
+
+// Returns the float value at a given point in the CSV file
+bool UHoudiniCSV::GetCSVFloatValue( const int32& lineIndex, const FString& ColumnTitle, float& value )
+{
+    int32 ColIndex = -1;
+    if ( !GetColumnIndexFromString( ColumnTitle, ColIndex ) )
+	return false;
+
+    return GetCSVFloatValue( lineIndex, ColIndex, value );
+}
+
+
+// Returns the string value at a given point in the CSV file
+bool UHoudiniCSV::GetCSVStringValue( const int32& lineIndex, const FString& ColumnTitle, FString& value )
+{
+    int32 ColIndex = -1;
+    if ( !GetColumnIndexFromString( ColumnTitle, ColIndex ) )
+	return false;
+
+    return GetCSVStringValue( lineIndex, ColIndex, value );
+}
+
+#if WITH_EDITORONLY_DATA
+void UHoudiniCSV::PostInitProperties()
+{
+    if ( !HasAnyFlags( RF_ClassDefaultObject ) )
+    {
+	AssetImportData = NewObject<UAssetImportData>(this, TEXT("AssetImportData"));
+    }
+
+    Super::PostInitProperties();
+}
+#endif
+
+#undef LOCTEXT_NAMESPACE
diff --git a/Engine/Plugins/FX/HoudiniNiagara/Source/HoudiniNiagara/Private/HoudiniNiagara.cpp b/Engine/Plugins/FX/HoudiniNiagara/Source/HoudiniNiagara/Private/HoudiniNiagara.cpp
new file mode 100644
--- /dev/null
+++ b/Engine/Plugins/FX/HoudiniNiagara/Source/HoudiniNiagara/Private/HoudiniNiagara.cpp
@@ -0,0 +1,42 @@
+/*
+* Copyright (c) <2017> Side Effects Software Inc.
+*
+* Permission is hereby granted, free of charge, to any person obtaining a copy
+* of this software and associated documentation files (the "Software"), to deal
+* in the Software without restriction, including without limitation the rights
+* to use, copy, modify, merge, publish, distribute, sublicense, and/or sell
+* copies of the Software, and to permit persons to whom the Software is
+* furnished to do so, subject to the following conditions:
+*
+* The above copyright notice and this permission notice shall be included in all
+* copies or substantial portions of the Software.
+*
+* THE SOFTWARE IS PROVIDED "AS IS", WITHOUT WARRANTY OF ANY KIND, EXPRESS OR
+* IMPLIED, INCLUDING BUT NOT LIMITED TO THE WARRANTIES OF MERCHANTABILITY,
+* FITNESS FOR A PARTICULAR PURPOSE AND NONINFRINGEMENT. IN NO EVENT SHALL THE
+* AUTHORS OR COPYRIGHT HOLDERS BE LIABLE FOR ANY CLAIM, DAMAGES OR OTHER
+* LIABILITY, WHETHER IN AN ACTION OF CONTRACT, TORT OR OTHERWISE, ARISING FROM,
+* OUT OF OR IN CONNECTION WITH THE SOFTWARE OR THE USE OR OTHER DEALINGS IN THE
+* SOFTWARE.
+*
+*/
+
+#include "HoudiniNiagara.h"
+#include "Modules/ModuleManager.h"
+
+#define LOCTEXT_NAMESPACE "FHoudiniNiagaraModule"
+
+void FHoudiniNiagaraModule::StartupModule()
+{
+	// This code will execute after your module is loaded into memory; the exact timing is specified in the .uplugin file per-module
+}
+
+void FHoudiniNiagaraModule::ShutdownModule()
+{
+	// This function may be called during shutdown to clean up your module.  For modules that support dynamic reloading,
+	// we call this function before unloading the module.
+}
+
+#undef LOCTEXT_NAMESPACE
+	
+IMPLEMENT_MODULE(FHoudiniNiagaraModule, HoudiniNiagara)
diff --git a/Engine/Plugins/FX/HoudiniNiagara/Source/HoudiniNiagara/Private/NiagaraDataInterfaceHoudiniCSV.cpp b/Engine/Plugins/FX/HoudiniNiagara/Source/HoudiniNiagara/Private/NiagaraDataInterfaceHoudiniCSV.cpp
new file mode 100644
--- /dev/null
+++ b/Engine/Plugins/FX/HoudiniNiagara/Source/HoudiniNiagara/Private/NiagaraDataInterfaceHoudiniCSV.cpp
@@ -0,0 +1,661 @@
+/*
+* Copyright (c) <2017> Side Effects Software Inc.
+*
+* Permission is hereby granted, free of charge, to any person obtaining a copy
+* of this software and associated documentation files (the "Software"), to deal
+* in the Software without restriction, including without limitation the rights
+* to use, copy, modify, merge, publish, distribute, sublicense, and/or sell
+* copies of the Software, and to permit persons to whom the Software is
+* furnished to do so, subject to the following conditions:
+*
+* The above copyright notice and this permission notice shall be included in all
+* copies or substantial portions of the Software.
+*
+* THE SOFTWARE IS PROVIDED "AS IS", WITHOUT WARRANTY OF ANY KIND, EXPRESS OR
+* IMPLIED, INCLUDING BUT NOT LIMITED TO THE WARRANTIES OF MERCHANTABILITY,
+* FITNESS FOR A PARTICULAR PURPOSE AND NONINFRINGEMENT. IN NO EVENT SHALL THE
+* AUTHORS OR COPYRIGHT HOLDERS BE LIABLE FOR ANY CLAIM, DAMAGES OR OTHER
+* LIABILITY, WHETHER IN AN ACTION OF CONTRACT, TORT OR OTHERWISE, ARISING FROM,
+* OUT OF OR IN CONNECTION WITH THE SOFTWARE OR THE USE OR OTHER DEALINGS IN THE
+* SOFTWARE.
+*
+*/
+
+#include "NiagaraDataInterfaceHoudiniCSV.h"
+#include "NiagaraTypes.h"
+#include "Misc/FileHelper.h"
+
+#define LOCTEXT_NAMESPACE HOUDINI_NIAGARA_LOCTEXT_NAMESPACE 
+
+
+UNiagaraDataInterfaceHoudiniCSV::UNiagaraDataInterfaceHoudiniCSV(FObjectInitializer const& ObjectInitializer)
+	: Super(ObjectInitializer)
+{
+    CSVFile = nullptr;
+    LastSpawnIndex = -1;
+}
+
+void UNiagaraDataInterfaceHoudiniCSV::PostInitProperties()
+{
+    Super::PostInitProperties();
+
+    if (HasAnyFlags(RF_ClassDefaultObject))
+    {
+	    FNiagaraTypeRegistry::Register(FNiagaraTypeDefinition(GetClass()), true, false, false);
+    }
+
+    GPUBufferDirty = true;
+    LastSpawnIndex = -1;
+}
+
+void UNiagaraDataInterfaceHoudiniCSV::PostLoad()
+{
+    Super::PostLoad();
+    GPUBufferDirty = true;
+    LastSpawnIndex = -1;
+}
+
+#if WITH_EDITOR
+
+void UNiagaraDataInterfaceHoudiniCSV::PostEditChangeProperty(struct FPropertyChangedEvent& PropertyChangedEvent)
+{
+    Super::PostEditChangeProperty(PropertyChangedEvent);
+
+    if (PropertyChangedEvent.Property && PropertyChangedEvent.Property->GetFName() == GET_MEMBER_NAME_CHECKED(UNiagaraDataInterfaceHoudiniCSV, CSVFile))
+    {
+	Modify();
+	if ( CSVFile )
+	{
+	    GPUBufferDirty = true;
+	    LastSpawnIndex = -1;
+	}
+    }
+}
+
+#endif
+
+bool UNiagaraDataInterfaceHoudiniCSV::CopyToInternal(UNiagaraDataInterface* Destination) const
+{
+    if ( !Super::CopyToInternal( Destination ) )
+    {
+	return false;
+    }
+
+    UNiagaraDataInterfaceHoudiniCSV* CastedInterface = CastChecked<UNiagaraDataInterfaceHoudiniCSV>( Destination );
+    if ( !CastedInterface )
+	return false;
+
+    CastedInterface->CSVFile = CSVFile;
+
+    return true;
+}
+
+bool UNiagaraDataInterfaceHoudiniCSV::Equals(const UNiagaraDataInterface* Other) const
+{
+    if ( !Super::Equals(Other) )
+    {
+	return false;
+    }
+
+    const UNiagaraDataInterfaceHoudiniCSV* OtherHNCSV = CastChecked<UNiagaraDataInterfaceHoudiniCSV>(Other);
+
+    if ( OtherHNCSV != nullptr && OtherHNCSV->CSVFile != nullptr && CSVFile )
+    {
+	// Just make sure the two interfaces point to the same file
+	return OtherHNCSV->CSVFile->FileName.Equals( CSVFile->FileName );
+    }
+
+    return false;
+}
+
+
+void UNiagaraDataInterfaceHoudiniCSV::GetFunctions(TArray<FNiagaraFunctionSignature>& OutFunctions)
+{
+    {
+	// GetCSVFloatValue
+	FNiagaraFunctionSignature Sig;
+	Sig.Name = TEXT("GetCSVFloatValue");
+	Sig.bMemberFunction = true;
+	Sig.bRequiresContext = false;
+	Sig.Inputs.Add(FNiagaraVariable(FNiagaraTypeDefinition(GetClass()), TEXT("CSV")));		// CSV in
+	Sig.Inputs.Add(FNiagaraVariable(FNiagaraTypeDefinition::GetFloatDef(), TEXT("Row")));		// Row Index In
+	Sig.Inputs.Add(FNiagaraVariable(FNiagaraTypeDefinition::GetFloatDef(), TEXT("Col")));		// Col Index In
+	Sig.Outputs.Add(FNiagaraVariable(FNiagaraTypeDefinition::GetFloatDef(), TEXT("Value")));	// Float Out
+
+	OutFunctions.Add(Sig);
+    }
+
+    {
+	// GetCSVFloatValueByString
+	FNiagaraFunctionSignature Sig;
+	Sig.Name = TEXT("GetCSVFloatValueByString");
+	Sig.bMemberFunction = true;
+	Sig.bRequiresContext = false;
+	Sig.Inputs.Add(FNiagaraVariable(FNiagaraTypeDefinition(GetClass()), TEXT("CSV")));		// CSV in
+	Sig.Inputs.Add(FNiagaraVariable(FNiagaraTypeDefinition::GetFloatDef(), TEXT("Row")));		// Row Index In
+	Sig.Inputs.Add(FNiagaraVariable(FNiagaraTypeDefinition(GetClass()), TEXT("ColTitle")));	// Col Title In
+	Sig.Outputs.Add(FNiagaraVariable(FNiagaraTypeDefinition::GetFloatDef(), TEXT("Value")));	// Float Out
+
+	OutFunctions.Add(Sig);
+    }
+
+    {
+	// GetCSVPosition
+	FNiagaraFunctionSignature Sig;
+	Sig.Name = TEXT("GetCSVPosition");
+	Sig.bMemberFunction = true;
+	Sig.bRequiresContext = false;
+	Sig.Inputs.Add(FNiagaraVariable(FNiagaraTypeDefinition(GetClass()), TEXT("CSV")));		// CSV in
+	Sig.Inputs.Add(FNiagaraVariable(FNiagaraTypeDefinition::GetFloatDef(), TEXT("N")));		// Point Number In
+	Sig.Outputs.Add(FNiagaraVariable(FNiagaraTypeDefinition::GetVec3Def(), TEXT("Value")));		// Vector3 Out
+
+	OutFunctions.Add(Sig);
+    }
+
+    {
+	// GetCSVPositionAndTime
+	FNiagaraFunctionSignature Sig;
+	Sig.Name = TEXT("GetCSVPositionAndTime");
+	Sig.bMemberFunction = true;
+	Sig.bRequiresContext = false;
+	Sig.Inputs.Add(FNiagaraVariable(FNiagaraTypeDefinition(GetClass()), TEXT("CSV")));		// CSV in
+	Sig.Inputs.Add(FNiagaraVariable(FNiagaraTypeDefinition::GetFloatDef(), TEXT("N")));		// Point Number In
+	Sig.Outputs.Add(FNiagaraVariable(FNiagaraTypeDefinition::GetVec3Def(), TEXT("Position")));	// Vector3 Out
+	Sig.Outputs.Add(FNiagaraVariable(FNiagaraTypeDefinition::GetFloatDef(), TEXT("Time")));		// float Out
+
+	OutFunctions.Add(Sig);
+    }
+
+    {
+	// GetCSVNormal
+	FNiagaraFunctionSignature Sig;
+	Sig.Name = TEXT("GetCSVNormal");
+	Sig.bMemberFunction = true;
+	Sig.bRequiresContext = false;
+	Sig.Inputs.Add(FNiagaraVariable(FNiagaraTypeDefinition(GetClass()), TEXT("CSV")));		// CSV in
+	Sig.Inputs.Add(FNiagaraVariable(FNiagaraTypeDefinition::GetFloatDef(), TEXT("N")));		// Point Number In
+	Sig.Outputs.Add(FNiagaraVariable(FNiagaraTypeDefinition::GetVec3Def(), TEXT("Value")));		// Vector3 Out
+
+	OutFunctions.Add(Sig);
+    }
+
+    {
+	// GetCSVTime
+	FNiagaraFunctionSignature Sig;
+	Sig.Name = TEXT("GetCSVTime");
+	Sig.bMemberFunction = true;
+	Sig.bRequiresContext = false;
+	Sig.Inputs.Add(FNiagaraVariable(FNiagaraTypeDefinition(GetClass()), TEXT("CSV")));		// CSV in
+	Sig.Inputs.Add(FNiagaraVariable(FNiagaraTypeDefinition::GetFloatDef(), TEXT("N")));		// Point Number In
+	Sig.Outputs.Add(FNiagaraVariable(FNiagaraTypeDefinition::GetFloatDef(), TEXT("Value")));	// Float Out
+
+	OutFunctions.Add(Sig);
+    }
+
+    {
+	// GetNumberOfPointsInCSV
+	FNiagaraFunctionSignature Sig;
+	Sig.Name = TEXT("GetNumberOfPointsInCSV");
+	Sig.bMemberFunction = true;
+	Sig.bRequiresContext = false;
+	Sig.Inputs.Add(FNiagaraVariable(FNiagaraTypeDefinition(GetClass()), TEXT("CSV")));		    // CSV in
+	Sig.Outputs.Add(FNiagaraVariable(FNiagaraTypeDefinition::GetIntDef(), TEXT("NumberOfPoints")));     // Int Out
+
+	OutFunctions.Add(Sig);
+    }
+
+    {
+	// GetLastParticleIndexAtTime
+	FNiagaraFunctionSignature Sig;
+	Sig.Name = TEXT("GetLastParticleIndexAtTime");
+	Sig.bMemberFunction = true;
+	Sig.bRequiresContext = false;
+	Sig.Inputs.Add(FNiagaraVariable(FNiagaraTypeDefinition(GetClass()), TEXT("CSV")));		    // CSV in
+	Sig.Inputs.Add(FNiagaraVariable(FNiagaraTypeDefinition::GetFloatDef(), TEXT("Time")));		    // Time in
+	Sig.Outputs.Add(FNiagaraVariable(FNiagaraTypeDefinition::GetIntDef(), TEXT("LastIndex")));	    // Int Out
+
+	OutFunctions.Add(Sig);
+    }
+
+    {
+	// GetParticleIndexesToSpawnAtTime
+	FNiagaraFunctionSignature Sig;
+	Sig.Name = TEXT("GetParticleIndexesToSpawnAtTime");
+	Sig.bMemberFunction = true;
+	Sig.bRequiresContext = false;
+	Sig.Inputs.Add(FNiagaraVariable(FNiagaraTypeDefinition(GetClass()), TEXT("CSV")));		    // CSV in
+	Sig.Inputs.Add(FNiagaraVariable(FNiagaraTypeDefinition::GetFloatDef(), TEXT("Time")));		    // Time in
+	Sig.Outputs.Add(FNiagaraVariable(FNiagaraTypeDefinition::GetIntDef(), TEXT("MinIndex")));	    // Int Out
+	Sig.Outputs.Add(FNiagaraVariable(FNiagaraTypeDefinition::GetIntDef(), TEXT("MaxIndex")));	    // Int Out
+	Sig.Outputs.Add(FNiagaraVariable(FNiagaraTypeDefinition::GetIntDef(), TEXT("Count")));		    // Int Out
+
+	OutFunctions.Add(Sig);
+    }
+
+}
+
+// build the shader function HLSL; function name is passed in, as it's defined per-DI; that way, configuration could change
+// the HLSL in the spirit of a static switch
+// TODO: need a way to identify each specific function here
+// 
+bool UNiagaraDataInterfaceHoudiniCSV::GetFunctionHLSL(const FName& DefinitionFunctionName, FString InstanceFunctionName, FNiagaraDataInterfaceGPUParamInfo& ParamInfo, FString& OutHLSL)
+{
+    //FString BufferName = Descriptors[0].BufferParamName;
+    //FString SecondBufferName = Descriptors[1].BufferParamName;
+
+    /*
+    if (InstanceFunctionName.Contains( TEXT( "GetCSVFloatValue") ) )
+    {
+	FString BufferName = Descriptors[0].BufferParamName;
+	OutHLSL += TEXT("void ") + InstanceFunctionName + TEXT("(in float In_Row, in float In_Col, out float Out_Value) \n{\n");
+	OutHLSL += TEXT("\t Out_Value = ") + BufferName + TEXT("[(int)(In_Row + ( In_Col * ") + FString::FromInt(NumberOfRows) + TEXT(") ) ];");
+	OutHLSL += TEXT("\n}\n");
+    }
+    else if (InstanceFunctionName.Contains(TEXT("GetCSVPosition")))
+    {
+	FString BufferName = Descriptors[1].BufferParamName;
+	OutHLSL += TEXT("void ") + InstanceFunctionName + TEXT("(in float In_N, out float3 Out_Value) \n{\n");
+	OutHLSL += TEXT("\t Out_Value.x = ") + BufferName + TEXT("[(int)(In_N) ];");
+	OutHLSL += TEXT("\t Out_Value.y = ") + BufferName + TEXT("[(int)(In_N + ") + FString::FromInt(NumberOfRows) + TEXT(") ];");
+	OutHLSL += TEXT("\t Out_Value.z = ") + BufferName + TEXT("[(int)(In_N + ( 2 * ") + FString::FromInt(NumberOfRows) + TEXT(") ) ];");
+	OutHLSL += TEXT("\n}\n");
+    }
+    else if (InstanceFunctionName.Contains(TEXT("GetCSVNormal")))
+    {
+	FString BufferName = Descriptors[2].BufferParamName;
+	OutHLSL += TEXT("void ") + InstanceFunctionName + TEXT("(in float In_N, out float3 Out_Value) \n{\n");
+	OutHLSL += TEXT("\t Out_Value.x = ") + BufferName + TEXT("[(int)(In_N) ];");
+	OutHLSL += TEXT("\t Out_Value.y = ") + BufferName + TEXT("[(int)(In_N + ") + FString::FromInt(NumberOfRows) + TEXT(") ];");
+	OutHLSL += TEXT("\t Out_Value.z = ") + BufferName + TEXT("[(int)(In_N + ( 2 * ") + FString::FromInt(NumberOfRows) + TEXT(") ) ];");
+	OutHLSL += TEXT("\n}\n");
+    }
+    else if (InstanceFunctionName.Contains(TEXT("GetCSVTime")))
+    {
+	FString BufferName = Descriptors[3].BufferParamName;
+	OutHLSL += TEXT("void ") + InstanceFunctionName + TEXT("(in float In_N, out float Out_Value) \n{\n");
+	OutHLSL += TEXT("\t Out_Value = ") + BufferName + TEXT("[(int)(In_N) ];");
+	OutHLSL += TEXT("\n}\n");
+    }
+    else if (InstanceFunctionName.Contains( TEXT("GetNumberOfPointsInCSV") ) )
+    {
+	OutHLSL += TEXT("void ") + InstanceFunctionName + TEXT("( out int Out_Value ) \n{\n");
+	OutHLSL += TEXT("\t Out_Value = ") + FString::FromInt(NumberOfRows) + TEXT(";");
+	OutHLSL += TEXT("\n}\n");
+    }*/
+    /*else if (InstanceFunctionName.Contains(TEXT("GetCSVPositionAndTime")))
+    {
+	OutHLSL += TEXT("void ") + FunctionName + TEXT("(in float In_N, out float4 Out_Value) \n{\n");
+	OutHLSL += TEXT("\t Out_Value.x = ") + BufferName + TEXT("[(int)(In_N) ];");
+	OutHLSL += TEXT("\t Out_Value.y = ") + BufferName + TEXT("[(int)(In_N + ") + FString::FromInt(NumberOfRows) + TEXT(") ];");
+	OutHLSL += TEXT("\t Out_Value.z = ") + BufferName + TEXT("[(int)(In_N + ( 2 * ") + FString::FromInt(NumberOfRows) + TEXT(") ) ];");
+	OutHLSL += TEXT("\t Out_Value.w = ") + SecondBufferName + TEXT("[(int)(In_N) ];");
+	OutHLSL += TEXT("\n}\n");
+    }*/
+
+    return !OutHLSL.IsEmpty();
+}
+
+// build buffer definition hlsl
+// 1. Choose a buffer name, add the data interface ID (important!)
+// 2. add a DIGPUBufferParamDescriptor to the array argument; that'll be passed on to the FNiagaraShader for binding to a shader param, that can
+// then later be found by name via FindDIBufferParam for setting; 
+// 3. store buffer declaration hlsl in OutHLSL
+// multiple buffers can be defined at once here
+//
+void UNiagaraDataInterfaceHoudiniCSV::GetParameterDefinitionHLSL(FNiagaraDataInterfaceGPUParamInfo& ParamInfo, FString& OutHLSL)
+{
+    FString BufferName = "CSVData_" + ParamInfo.DataInterfaceHLSLSymbol;
+    OutHLSL += TEXT("Buffer<float> ") + BufferName + TEXT(";\n");
+
+    BufferName = "PositionData_" + ParamInfo.DataInterfaceHLSLSymbol;
+    OutHLSL += TEXT("Buffer<float> ") + BufferName + TEXT(";\n");
+
+    BufferName = "NormalData_" + ParamInfo.DataInterfaceHLSLSymbol;
+    OutHLSL += TEXT("Buffer<float> ") + BufferName + TEXT(";\n");
+
+    BufferName = "TimeData_" + ParamInfo.DataInterfaceHLSLSymbol;
+    OutHLSL += TEXT("Buffer<float> ") + BufferName + TEXT(";\n");
+}
+
+DEFINE_NDI_RAW_FUNC_BINDER(UNiagaraDataInterfaceHoudiniCSV, GetCSVFloatValue);
+DEFINE_NDI_RAW_FUNC_BINDER(UNiagaraDataInterfaceHoudiniCSV, GetCSVFloatValueByString);
+DEFINE_NDI_RAW_FUNC_BINDER(UNiagaraDataInterfaceHoudiniCSV, GetCSVVectorValue);
+DEFINE_NDI_RAW_FUNC_BINDER(UNiagaraDataInterfaceHoudiniCSV, GetCSVPosition);
+DEFINE_NDI_RAW_FUNC_BINDER(UNiagaraDataInterfaceHoudiniCSV, GetCSVNormal);
+DEFINE_NDI_RAW_FUNC_BINDER(UNiagaraDataInterfaceHoudiniCSV, GetCSVTime);
+DEFINE_NDI_RAW_FUNC_BINDER(UNiagaraDataInterfaceHoudiniCSV, GetCSVPositionAndTime);
+DEFINE_NDI_RAW_FUNC_BINDER(UNiagaraDataInterfaceHoudiniCSV, GetLastParticleIndexAtTime);
+DEFINE_NDI_RAW_FUNC_BINDER(UNiagaraDataInterfaceHoudiniCSV, GetParticleIndexesToSpawnAtTime);
+void UNiagaraDataInterfaceHoudiniCSV::GetVMExternalFunction(const FVMExternalFunctionBindingInfo& BindingInfo, void* InstanceData, FVMExternalFunction &OutFunc)
+{
+    if (BindingInfo.Name == TEXT("GetCSVFloatValue") && BindingInfo.GetNumInputs() == 2 && BindingInfo.GetNumOutputs() == 1)
+    {
+	TNDIParamBinder<0, float, TNDIParamBinder<1, float, NDI_RAW_FUNC_BINDER(UNiagaraDataInterfaceHoudiniCSV, GetCSVFloatValue)>>::Bind(this, BindingInfo, InstanceData, OutFunc);
+    }
+    else if (BindingInfo.Name == TEXT("GetCSVFloatValueByString") && BindingInfo.GetNumInputs() == 2 && BindingInfo.GetNumOutputs() == 1)
+    {
+	TNDIParamBinder<0, float, TNDIParamBinder<1, FString, NDI_RAW_FUNC_BINDER(UNiagaraDataInterfaceHoudiniCSV, GetCSVFloatValueByString)>>::Bind(this, BindingInfo, InstanceData, OutFunc);
+    }
+    else if (BindingInfo.Name == TEXT("GetCSVVectorValue") && BindingInfo.GetNumInputs() == 2 && BindingInfo.GetNumOutputs() == 1)
+    {
+	TNDIParamBinder<0, float, TNDIParamBinder<1, float, NDI_RAW_FUNC_BINDER(UNiagaraDataInterfaceHoudiniCSV, GetCSVVectorValue)>>::Bind(this, BindingInfo, InstanceData, OutFunc);
+    }    
+    else if (BindingInfo.Name == TEXT("GetCSVPosition") && BindingInfo.GetNumInputs() == 1 && BindingInfo.GetNumOutputs() == 3)
+    {
+	TNDIParamBinder<0, float, NDI_RAW_FUNC_BINDER(UNiagaraDataInterfaceHoudiniCSV, GetCSVPosition)>::Bind(this, BindingInfo, InstanceData, OutFunc);
+    }
+    else if (BindingInfo.Name == TEXT("GetCSVNormal") && BindingInfo.GetNumInputs() == 1 && BindingInfo.GetNumOutputs() == 3)
+    {
+	TNDIParamBinder<0, float, NDI_RAW_FUNC_BINDER(UNiagaraDataInterfaceHoudiniCSV, GetCSVNormal)>::Bind(this, BindingInfo, InstanceData, OutFunc);
+    }
+    else if (BindingInfo.Name == TEXT("GetCSVTime") && BindingInfo.GetNumInputs() == 1 && BindingInfo.GetNumOutputs() == 1)
+    {
+	TNDIParamBinder<0, float, NDI_RAW_FUNC_BINDER(UNiagaraDataInterfaceHoudiniCSV, GetCSVTime)>::Bind(this, BindingInfo, InstanceData, OutFunc);
+    }
+    else if (BindingInfo.Name == TEXT("GetCSVPositionAndTime") && BindingInfo.GetNumInputs() == 1 && BindingInfo.GetNumOutputs() == 4)
+    {
+	TNDIParamBinder<0, float, NDI_RAW_FUNC_BINDER(UNiagaraDataInterfaceHoudiniCSV, GetCSVPositionAndTime)>::Bind(this, BindingInfo, InstanceData, OutFunc);
+    }
+    else if ( BindingInfo.Name == TEXT("GetNumberOfPointsInCSV") && BindingInfo.GetNumInputs() == 0 && BindingInfo.GetNumOutputs() == 1 )
+    {
+	OutFunc = FVMExternalFunction::CreateUObject(this, &UNiagaraDataInterfaceHoudiniCSV::GetNumberOfPointsInCSV);
+    }
+    else if (BindingInfo.Name == TEXT("GetLastParticleIndexAtTime") && BindingInfo.GetNumInputs() == 1 && BindingInfo.GetNumOutputs() == 1)
+    {
+	TNDIParamBinder<0, float, NDI_RAW_FUNC_BINDER(UNiagaraDataInterfaceHoudiniCSV, GetLastParticleIndexAtTime)>::Bind(this, BindingInfo, InstanceData, OutFunc);
+    }
+    else if (BindingInfo.Name == TEXT("GetParticleIndexesToSpawnAtTime") && BindingInfo.GetNumInputs() == 1 && BindingInfo.GetNumOutputs() == 3)
+    {
+	TNDIParamBinder<0, float, NDI_RAW_FUNC_BINDER(UNiagaraDataInterfaceHoudiniCSV, GetParticleIndexesToSpawnAtTime)>::Bind(this, BindingInfo, InstanceData, OutFunc);
+    }
+    else
+    {
+	UE_LOG( LogHoudiniNiagara, Error, 
+	    TEXT( "Could not find data interface function:\n\tName: %s\n\tInputs: %i\n\tOutputs: %i" ),
+	    *BindingInfo.Name.ToString(), BindingInfo.GetNumInputs(), BindingInfo.GetNumOutputs() );
+	OutFunc = FVMExternalFunction();
+    }
+}
+
+template<typename RowParamType, typename ColParamType>
+void UNiagaraDataInterfaceHoudiniCSV::GetCSVFloatValue(FVectorVMContext& Context)
+{
+    RowParamType RowParam(Context);
+    ColParamType ColParam(Context);
+
+    FRegisterHandler<float> OutValue(Context);
+
+    for ( int32 i = 0; i < Context.NumInstances; ++i )
+    {
+	float row = RowParam.Get();
+	float col = ColParam.Get();
+	
+	float value = 0.0f;
+	if ( CSVFile )
+	    CSVFile->GetCSVFloatValue( row, col, value );
+
+	*OutValue.GetDest() = value;
+	RowParam.Advance();
+	ColParam.Advance();
+	OutValue.Advance();
+    }
+}
+
+template<typename RowParamType, typename ColParamType>
+void UNiagaraDataInterfaceHoudiniCSV::GetCSVVectorValue( FVectorVMContext& Context )
+{
+    RowParamType RowParam(Context);
+    ColParamType ColParam(Context);
+
+    FRegisterHandler<FVector> OutValue(Context);
+
+    for (int32 i = 0; i < Context.NumInstances; ++i)
+    {
+	float row = RowParam.Get();
+	float col = ColParam.Get();
+
+	FVector value = FVector::ZeroVector;
+	if ( CSVFile )
+	    CSVFile->GetCSVVectorValue(row, col, value);
+
+	*OutValue.GetDest() = value;
+	RowParam.Advance();
+	ColParam.Advance();
+	OutValue.Advance();
+    }
+}
+
+template<typename RowParamType, typename ColTitleParamType>
+void UNiagaraDataInterfaceHoudiniCSV::GetCSVFloatValueByString(FVectorVMContext& Context)
+{
+    RowParamType RowParam(Context);
+    ColTitleParamType ColTitleParam(Context);
+
+    FRegisterHandler<float> OutValue(Context);
+
+    for ( int32 i = 0; i < Context.NumInstances; ++i )
+    {
+	float row = RowParam.Get();
+	FString colTitle = ColTitleParam.Get();
+	
+	float value = 0.0f;
+	if ( CSVFile )
+	    CSVFile->GetCSVFloatValue( row, colTitle, value );
+
+	*OutValue.GetDest() = value;
+	RowParam.Advance();
+	ColTitleParam.Advance();
+	OutValue.Advance();
+    }
+}
+
+template<typename NParamType>
+void UNiagaraDataInterfaceHoudiniCSV::GetCSVPosition(FVectorVMContext& Context)
+{
+    NParamType NParam(Context);
+    FRegisterHandler<float> OutSampleX(Context);
+    FRegisterHandler<float> OutSampleY(Context);
+    FRegisterHandler<float> OutSampleZ(Context);
+
+    for (int32 i = 0; i < Context.NumInstances; ++i)
+    {
+	float N = NParam.Get();
+
+	FVector V = FVector::ZeroVector;
+	if ( CSVFile )
+	    CSVFile->GetCSVPositionValue( N, V );
+
+	*OutSampleX.GetDest() = V.X;
+	*OutSampleY.GetDest() = V.Y;
+	*OutSampleZ.GetDest() = V.Z;
+	NParam.Advance();
+	OutSampleX.Advance();
+	OutSampleY.Advance();
+	OutSampleZ.Advance();
+    }
+}
+
+template<typename NParamType>
+void UNiagaraDataInterfaceHoudiniCSV::GetCSVNormal(FVectorVMContext& Context)
+{
+    NParamType NParam(Context);
+    FRegisterHandler<float> OutSampleX(Context);
+    FRegisterHandler<float> OutSampleY(Context);
+    FRegisterHandler<float> OutSampleZ(Context);
+
+    for (int32 i = 0; i < Context.NumInstances; ++i)
+    {
+	float N = NParam.Get();
+
+	FVector V = FVector::ZeroVector;
+	if ( CSVFile )
+	    CSVFile->GetCSVNormalValue( N, V );
+
+	*OutSampleX.GetDest() = V.X;
+	*OutSampleY.GetDest() = V.Y;
+	*OutSampleZ.GetDest() = V.Z;
+	NParam.Advance();
+	OutSampleX.Advance();
+	OutSampleY.Advance();
+	OutSampleZ.Advance();
+    }
+}
+
+template<typename NParamType>
+void UNiagaraDataInterfaceHoudiniCSV::GetCSVTime(FVectorVMContext& Context)
+{
+    NParamType NParam(Context);
+    FRegisterHandler<float> OutValue(Context);
+
+    for (int32 i = 0; i < Context.NumInstances; ++i)
+    {
+	float N = NParam.Get();
+
+	float value = 0.0f;
+	if ( CSVFile )
+	    CSVFile->GetCSVTimeValue( N, value );
+
+	*OutValue.GetDest() = value;
+	NParam.Advance();
+	OutValue.Advance();
+    }
+}
+
+// Returns the last index of the particles that should be spawned at time t
+template<typename TimeParamType>
+void UNiagaraDataInterfaceHoudiniCSV::GetLastParticleIndexAtTime(FVectorVMContext& Context)
+{
+    TimeParamType TimeParam(Context);
+    FRegisterHandler<int32> OutValue(Context);
+
+    for (int32 i = 0; i < Context.NumInstances; ++i)
+    {
+	float t = TimeParam.Get();
+
+	int32 value = 0;
+	if ( CSVFile )
+	    CSVFile->GetLastParticleIndexAtTime( t, value );
+
+	*OutValue.GetDest() = value;
+	TimeParam.Advance();
+	OutValue.Advance();
+    }
+}
+
+// Returns the last index of the particles that should be spawned at time t
+template<typename TimeParamType>
+void UNiagaraDataInterfaceHoudiniCSV::GetParticleIndexesToSpawnAtTime( FVectorVMContext& Context )
+{
+    TimeParamType TimeParam( Context );
+    FRegisterHandler<int32> OutMinValue( Context );
+    FRegisterHandler<int32> OutMaxValue( Context );
+    FRegisterHandler<int32> OutCountValue( Context );
+
+    for (int32 i = 0; i < Context.NumInstances; ++i)
+    {
+	float t = TimeParam.Get();
+
+	int32 value = 0;
+	int32 min = 0, max = 0, count = 0;
+	if ( CSVFile )
+	{
+	    if ( !CSVFile->GetLastParticleIndexAtTime( t, value ) )
+	    {
+		// The CSV file doesn't have time informations, so always return all points in the file
+		min = 0;
+		max = value;
+		count = max - min + 1;
+	    }
+	    else
+	    {
+		// The CSV file has time informations
+		// First, detect if we need to reset LastSpawnIndex (after a loop)
+		if ( value < LastSpawnIndex )
+		    LastSpawnIndex = -1;
+
+		if ( value < 0 )
+		{
+		    // Nothing to spawn, t is lower than the particle's time
+		    LastSpawnIndex = -1;
+		}
+		else
+		{
+		    // The last time value in the CSV is lower than t, spawn everything if we didnt already!
+		    if ( value >= CSVFile->GetNumberOfPointsInCSV() )
+			value = value - 1;
+
+		    if ( value == LastSpawnIndex )
+		    {
+			// We've already spawned all the particles for this time
+			min = value;
+			max = value;
+			count = 0;
+		    }
+		    else
+		    {
+			// Found particles to spawn at time t
+			min = LastSpawnIndex + 1;
+			max = value;
+			count = max - min + 1;
+
+			LastSpawnIndex = max;
+		    }
+		}
+	    }
+	}
+
+	*OutMinValue.GetDest() = min;
+	*OutMaxValue.GetDest() = max;
+	*OutCountValue.GetDest() = count;
+
+	TimeParam.Advance();
+	OutMinValue.Advance();
+	OutMaxValue.Advance();
+	OutCountValue.Advance();
+    }
+}
+
+template<typename NParamType>
+void UNiagaraDataInterfaceHoudiniCSV::GetCSVPositionAndTime(FVectorVMContext& Context)
+{
+    NParamType NParam(Context);
+
+    FRegisterHandler<float> OutPosX(Context);
+    FRegisterHandler<float> OutPosY(Context);
+    FRegisterHandler<float> OutPosZ(Context);
+    FRegisterHandler<float> OutTime(Context);
+
+    for (int32 i = 0; i < Context.NumInstances; ++i)
+    {
+	float N = NParam.Get();
+
+	float timeValue = 0.0f;
+	FVector posVector = FVector::ZeroVector;
+	if ( CSVFile )
+	{
+	    CSVFile->GetCSVTimeValue(N, timeValue);
+	    CSVFile->GetCSVPositionValue(N, posVector);
+	}
+
+	*OutPosX.GetDest() = posVector.X;
+	*OutPosY.GetDest() = posVector.Y;
+	*OutPosZ.GetDest() = posVector.Z;
+
+	*OutTime.GetDest() = timeValue;
+
+	NParam.Advance();
+	OutPosX.Advance();
+	OutPosY.Advance();
+	OutPosZ.Advance();
+	OutTime.Advance();
+    }
+}
+
+void UNiagaraDataInterfaceHoudiniCSV::GetNumberOfPointsInCSV(FVectorVMContext& Context)
+{
+    FRegisterHandler<int32> OutNumPoints(Context);
+    *OutNumPoints.GetDest() = CSVFile ? CSVFile->NumberOfLines : 0;
+    OutNumPoints.Advance();
+}
+
+#undef LOCTEXT_NAMESPACE
diff --git a/Engine/Plugins/FX/HoudiniNiagara/Source/HoudiniNiagara/Public/HoudiniCSV.h b/Engine/Plugins/FX/HoudiniNiagara/Source/HoudiniNiagara/Public/HoudiniCSV.h
new file mode 100644
--- /dev/null
+++ b/Engine/Plugins/FX/HoudiniNiagara/Source/HoudiniNiagara/Public/HoudiniCSV.h
@@ -0,0 +1,129 @@
+/*
+* Copyright (c) <2017> Side Effects Software Inc.
+*
+* Permission is hereby granted, free of charge, to any person obtaining a copy
+* of this software and associated documentation files (the "Software"), to deal
+* in the Software without restriction, including without limitation the rights
+* to use, copy, modify, merge, publish, distribute, sublicense, and/or sell
+* copies of the Software, and to permit persons to whom the Software is
+* furnished to do so, subject to the following conditions:
+*
+* The above copyright notice and this permission notice shall be included in all
+* copies or substantial portions of the Software.
+*
+* THE SOFTWARE IS PROVIDED "AS IS", WITHOUT WARRANTY OF ANY KIND, EXPRESS OR
+* IMPLIED, INCLUDING BUT NOT LIMITED TO THE WARRANTIES OF MERCHANTABILITY,
+* FITNESS FOR A PARTICULAR PURPOSE AND NONINFRINGEMENT. IN NO EVENT SHALL THE
+* AUTHORS OR COPYRIGHT HOLDERS BE LIABLE FOR ANY CLAIM, DAMAGES OR OTHER
+* LIABILITY, WHETHER IN AN ACTION OF CONTRACT, TORT OR OTHERWISE, ARISING FROM,
+* OUT OF OR IN CONNECTION WITH THE SOFTWARE OR THE USE OR OTHER DEALINGS IN THE
+* SOFTWARE.
+*
+*/
+#pragma once
+
+#include "UObject/ObjectMacros.h"
+#include "UObject/UObjectGlobals.h"
+#include "UObject/Object.h"
+#include "HoudiniCSV.generated.h"
+
+DECLARE_LOG_CATEGORY_EXTERN( LogHoudiniNiagara, All, All );
+
+UCLASS()
+class HOUDININIAGARA_API UHoudiniCSV : public UObject
+{
+    GENERATED_UCLASS_BODY()
+ 
+    public:
+	
+	//-----------------------------------------------------------------------------------------
+	//  MEMBER FUNCTIONS
+	//-----------------------------------------------------------------------------------------
+
+	bool UpdateFromFile( const FString& TheFileName );
+	bool UpdateFromStringArray( TArray<FString>& StringArray );
+
+	void SetFileName( const FString& TheFilename );
+
+	// Returns the number of points found in the CSV file
+	int32 GetNumberOfPointsInCSV();
+
+	// Returns the column index for a given string
+	bool GetColumnIndexFromString(const FString& ColumnTitle, int32& ColumnIndex);
+
+	// Returns the float value at a given point in the CSV file
+	bool GetCSVFloatValue( const int32& lineIndex, const int32& colIndex, float& value );
+	// Returns the float value at a given point in the CSV file
+	bool GetCSVFloatValue(const int32& lineIndex, const FString& ColumnTitle, float& value);
+	// Returns the float value at a given point in the CSV file
+	bool GetCSVStringValue( const int32& lineIndex, const int32& colIndex, FString& value );
+	// Returns the string value at a given point in the CSV file
+	bool GetCSVStringValue( const int32& lineIndex, const FString& ColumnTitle, FString& value );
+	// Returns a Vector3 for a given point in the CSV file
+	bool GetCSVVectorValue( const int32& lineIndex, const int32& colIndex, FVector& value, const bool& DoSwap = true, const bool& DoScale = true );
+	
+	// Returns a time value for a given point in the CSV file
+	bool GetCSVTimeValue( const int32& lineIndex, float& value );
+	// Returns a Position Vector3 for a given point in the CSV file (converted to unreal's coordinate system)
+	bool GetCSVPositionValue( const int32& lineIndex, FVector& value );
+	// Returns a Normal Vector3 for a given point in the CSV file (converted to unreal's coordinate system)
+	bool GetCSVNormalValue( const int32& lineIndex, FVector& value );
+
+	// Get the last index of the particles to be spawned at time t
+	// Invalid Index are used to indicate edge cases:
+	// -1 will be returned if no particles have been spawned ( t is smaller than the first particle time )
+	// NumberOfLines will be returned if all particles in the CSV have been spawned ( t is higher than the last particle time )
+	bool GetLastParticleIndexAtTime( const float& time, int32& lastIndex );
+
+	//-----------------------------------------------------------------------------------------
+	//  MEMBER VARIABLES
+	//-----------------------------------------------------------------------------------------
+	//UPROPERTY(EditAnywhere, Category = "My Object Properties")
+	UPROPERTY( VisibleAnywhere, Category = "Houdini CSV File Properties" )
+	FString FileName;
+
+	// The number of values stored in the CSV file (excluding the title row)
+	UPROPERTY( VisibleAnywhere, Category = "Houdini CSV File Properties" )
+	int32 NumberOfLines;
+
+	// The number of value TYPES stored in the CSV file
+	UPROPERTY( VisibleAnywhere, Category = "Houdini CSV File Properties" )
+	int32 NumberOfColumns;
+
+	// The tokenized title raw, describing the content of each column
+	UPROPERTY( VisibleAnywhere, Category = "Houdini CSV File Properties" )
+	TArray<FString> TitleRowArray;
+
+#if WITH_EDITORONLY_DATA
+	/** Importing data and options used for this mesh */
+	UPROPERTY(EditAnywhere, Instanced, Category = ImportSettings)
+	class UAssetImportData* AssetImportData;
+
+	virtual void PostInitProperties() override;
+#endif
+
+    private:
+
+	// Array containing the Raw String data
+	UPROPERTY()
+	TArray<FString> StringCSVData;
+
+	// Array containing all the CSV data converted to floats
+	UPROPERTY()
+	TArray<float> FloatCSVData;
+
+	// Array containing the Position buffer
+	UPROPERTY()	    
+	int32 PositionColumnIndex;
+	//TArray<float> PositionData;
+
+	// Array containing the Normal buffer
+	UPROPERTY()
+	int32 NormalColumnIndex;
+	//TArray<float> NormalData;
+
+	// Array containing the time buffer
+	UPROPERTY()
+	int32 TimeColumnIndex;
+	//TArray<float> TimeData;
+};
\ No newline at end of file
diff --git a/Engine/Plugins/FX/HoudiniNiagara/Source/HoudiniNiagara/Public/HoudiniNiagara.h b/Engine/Plugins/FX/HoudiniNiagara/Source/HoudiniNiagara/Public/HoudiniNiagara.h
new file mode 100644
--- /dev/null
+++ b/Engine/Plugins/FX/HoudiniNiagara/Source/HoudiniNiagara/Public/HoudiniNiagara.h
@@ -0,0 +1,36 @@
+/*
+* Copyright (c) <2017> Side Effects Software Inc.
+*
+* Permission is hereby granted, free of charge, to any person obtaining a copy
+* of this software and associated documentation files (the "Software"), to deal
+* in the Software without restriction, including without limitation the rights
+* to use, copy, modify, merge, publish, distribute, sublicense, and/or sell
+* copies of the Software, and to permit persons to whom the Software is
+* furnished to do so, subject to the following conditions:
+*
+* The above copyright notice and this permission notice shall be included in all
+* copies or substantial portions of the Software.
+*
+* THE SOFTWARE IS PROVIDED "AS IS", WITHOUT WARRANTY OF ANY KIND, EXPRESS OR
+* IMPLIED, INCLUDING BUT NOT LIMITED TO THE WARRANTIES OF MERCHANTABILITY,
+* FITNESS FOR A PARTICULAR PURPOSE AND NONINFRINGEMENT. IN NO EVENT SHALL THE
+* AUTHORS OR COPYRIGHT HOLDERS BE LIABLE FOR ANY CLAIM, DAMAGES OR OTHER
+* LIABILITY, WHETHER IN AN ACTION OF CONTRACT, TORT OR OTHERWISE, ARISING FROM,
+* OUT OF OR IN CONNECTION WITH THE SOFTWARE OR THE USE OR OTHER DEALINGS IN THE
+* SOFTWARE.
+*
+*/
+
+#pragma once
+
+#include "CoreMinimal.h"
+#include "Modules/ModuleInterface.h"
+
+class FHoudiniNiagaraModule : public IModuleInterface
+{
+public:
+
+	/** IModuleInterface implementation */
+	virtual void StartupModule() override;
+	virtual void ShutdownModule() override;
+};
\ No newline at end of file
diff --git a/Engine/Plugins/FX/HoudiniNiagara/Source/HoudiniNiagaraEditor/HoudiniNiagaraEditor.Build.cs b/Engine/Plugins/FX/HoudiniNiagara/Source/HoudiniNiagaraEditor/HoudiniNiagaraEditor.Build.cs
new file mode 100644
--- /dev/null
+++ b/Engine/Plugins/FX/HoudiniNiagara/Source/HoudiniNiagaraEditor/HoudiniNiagaraEditor.Build.cs
@@ -0,0 +1,93 @@
+/*
+* Copyright (c) <2017> Side Effects Software Inc.
+*
+* Permission is hereby granted, free of charge, to any person obtaining a copy
+* of this software and associated documentation files (the "Software"), to deal
+* in the Software without restriction, including without limitation the rights
+* to use, copy, modify, merge, publish, distribute, sublicense, and/or sell
+* copies of the Software, and to permit persons to whom the Software is
+* furnished to do so, subject to the following conditions:
+*
+* The above copyright notice and this permission notice shall be included in all
+* copies or substantial portions of the Software.
+*
+* THE SOFTWARE IS PROVIDED "AS IS", WITHOUT WARRANTY OF ANY KIND, EXPRESS OR
+* IMPLIED, INCLUDING BUT NOT LIMITED TO THE WARRANTIES OF MERCHANTABILITY,
+* FITNESS FOR A PARTICULAR PURPOSE AND NONINFRINGEMENT. IN NO EVENT SHALL THE
+* AUTHORS OR COPYRIGHT HOLDERS BE LIABLE FOR ANY CLAIM, DAMAGES OR OTHER
+* LIABILITY, WHETHER IN AN ACTION OF CONTRACT, TORT OR OTHERWISE, ARISING FROM,
+* OUT OF OR IN CONNECTION WITH THE SOFTWARE OR THE USE OR OTHER DEALINGS IN THE
+* SOFTWARE.
+*
+*/
+
+using UnrealBuildTool;
+
+public class HoudiniNiagaraEditor : ModuleRules
+{
+	public HoudiniNiagaraEditor(ReadOnlyTargetRules Target) : base(Target)
+	{
+		PCHUsage = ModuleRules.PCHUsageMode.UseExplicitOrSharedPCHs;
+		
+		
+		PrivateIncludePaths.AddRange(
+			new string[] 
+            {
+				"HoudiniNiagaraEditor/Private",
+            }
+        );
+			
+		
+		PublicDependencyModuleNames.AddRange(
+			new string[]
+			{
+				"Core",
+                "Niagara",
+                "NiagaraShader",
+				"CoreUObject",
+                "VectorVM",
+                "RHI",
+                "UtilityShaders",
+                "NiagaraVertexFactories",
+                "ShaderCore",
+                "HoudiniNiagara",
+                "UnrealEd",
+                                "EditorStyle",
+            }
+        );
+			
+		
+		PrivateDependencyModuleNames.AddRange(
+			new string[]
+			{
+				"CoreUObject",
+				"Engine",
+				"Slate",
+				"SlateCore",
+                "Niagara",
+                "NiagaraShader",
+                "RenderCore",
+                "UtilityShaders",
+                "ShaderCore",
+                "HoudiniNiagara"
+			}
+        );
+
+        PrivateIncludePathModuleNames.AddRange(
+            new string[] 
+            {
+                "Engine",
+                "LevelEditor",
+                "AssetTools",
+                "ContentBrowser",
+            }
+        );
+
+        DynamicallyLoadedModuleNames.AddRange(
+		    new string[]
+		    {
+			    // ... add any modules that your module loads dynamically here ...
+		    }
+        );
+	}
+}
diff --git a/Engine/Plugins/FX/HoudiniNiagara/Source/HoudiniNiagaraEditor/Private/HoudiniCSVAssetActions.cpp b/Engine/Plugins/FX/HoudiniNiagara/Source/HoudiniNiagaraEditor/Private/HoudiniCSVAssetActions.cpp
new file mode 100644
--- /dev/null
+++ b/Engine/Plugins/FX/HoudiniNiagara/Source/HoudiniNiagaraEditor/Private/HoudiniCSVAssetActions.cpp
@@ -0,0 +1,179 @@
+/*
+* Copyright (c) <2017> Side Effects Software Inc.
+*
+* Permission is hereby granted, free of charge, to any person obtaining a copy
+* of this software and associated documentation files (the "Software"), to deal
+* in the Software without restriction, including without limitation the rights
+* to use, copy, modify, merge, publish, distribute, sublicense, and/or sell
+* copies of the Software, and to permit persons to whom the Software is
+* furnished to do so, subject to the following conditions:
+*
+* The above copyright notice and this permission notice shall be included in all
+* copies or substantial portions of the Software.
+*
+* THE SOFTWARE IS PROVIDED "AS IS", WITHOUT WARRANTY OF ANY KIND, EXPRESS OR
+* IMPLIED, INCLUDING BUT NOT LIMITED TO THE WARRANTIES OF MERCHANTABILITY,
+* FITNESS FOR A PARTICULAR PURPOSE AND NONINFRINGEMENT. IN NO EVENT SHALL THE
+* AUTHORS OR COPYRIGHT HOLDERS BE LIABLE FOR ANY CLAIM, DAMAGES OR OTHER
+* LIABILITY, WHETHER IN AN ACTION OF CONTRACT, TORT OR OTHERWISE, ARISING FROM,
+* OUT OF OR IN CONNECTION WITH THE SOFTWARE OR THE USE OR OTHER DEALINGS IN THE
+* SOFTWARE.
+*
+*/
+#include "HoudiniCSVAssetActions.h"
+
+#include "Framework/MultiBox/MultiBoxBuilder.h"
+#include "HoudiniCSV.h"
+#include "EditorStyleSet.h"
+#include "Toolkits/AssetEditorToolkit.h"
+#include "EditorReimportHandler.h"
+
+//#include "HoudiniCSVAssetEditorToolkit.h"
+
+#define LOCTEXT_NAMESPACE "AssetTypeActions"
+
+
+FHoudiniCSVAssetActions::FHoudiniCSVAssetActions()
+{ }
+
+bool FHoudiniCSVAssetActions::CanFilter()
+{
+    return true;
+}
+
+
+void FHoudiniCSVAssetActions::GetActions(const TArray<UObject*>& InObjects, FMenuBuilder& MenuBuilder)
+{
+    FAssetTypeActions_Base::GetActions(InObjects, MenuBuilder);
+
+    auto HoudiniCSVAssets = GetTypedWeakObjectPtrs<UHoudiniCSV>(InObjects);
+
+    MenuBuilder.AddMenuEntry(
+	LOCTEXT("ReimportHoudiniCSVLabel", "Reimport"),
+	LOCTEXT("ReimportHoudiniCSVTooltip", "Reimport the selected Houdini CSV file(s)."),
+	FSlateIcon(FEditorStyle::GetStyleSetName(), "ContentBrowser.AssetActions.ReimportAsset"),
+	FUIAction(
+	    FExecuteAction::CreateSP(this, &FHoudiniCSVAssetActions::ExecuteReimport, HoudiniCSVAssets),
+	    FCanExecuteAction::CreateSP(this, &FHoudiniCSVAssetActions::CanExecuteReimport, HoudiniCSVAssets)
+	)
+    );
+
+    MenuBuilder.AddMenuEntry(
+	LOCTEXT("OpenHoudiniCSVLabel", "Open in Text Editor"),
+	LOCTEXT("OpenHoudiniCSVTooltip", "Open the selected Houdini CSV file(s) in a Text Editor."),
+	FSlateIcon(FEditorStyle::GetStyleSetName(), "ContentBrowser.AssetActions.OpenInExternalEditor"),
+	FUIAction(
+	    FExecuteAction::CreateSP(this, &FHoudiniCSVAssetActions::ExecuteOpenInEditor, HoudiniCSVAssets),
+	    FCanExecuteAction::CreateSP(this, &FHoudiniCSVAssetActions::CanExecuteOpenInEditor, HoudiniCSVAssets)
+	)
+    );
+}
+
+
+uint32 FHoudiniCSVAssetActions::GetCategories()
+{
+    return EAssetTypeCategories::Misc;
+}
+
+
+FText FHoudiniCSVAssetActions::GetName() const
+{
+    return NSLOCTEXT("AssetTypeActions", "AssetTypeActions_HoudiniCSV", "Houdini CSV Asset");
+}
+
+
+UClass* FHoudiniCSVAssetActions::GetSupportedClass() const
+{
+    return UHoudiniCSV::StaticClass();
+}
+
+FText FHoudiniCSVAssetActions::GetAssetDescription(const FAssetData& AssetData) const
+{
+    return FText::GetEmpty();
+}
+
+
+FColor FHoudiniCSVAssetActions::GetTypeColor() const
+{
+    return FColor::Orange;
+}
+
+
+bool FHoudiniCSVAssetActions::HasActions(const TArray<UObject*>& InObjects) const
+{
+    return true;
+}
+
+
+class UThumbnailInfo* FHoudiniCSVAssetActions::GetThumbnailInfo(UObject* Asset) const
+{
+    /*
+    UHoudiniCSV* HoudiniCSV = CastChecked<UHoudiniCSV>(Asset);
+    UThumbnailInfo* ThumbnailInfo = HoudiniCSV->ThumbnailInfo;
+    if (ThumbnailInfo == NULL)
+    {
+	ThumbnailInfo = NewObject<USceneThumbnailInfo>(HoudiniCSV, NAME_None, RF_Transactional);
+	HoudiniCSV->ThumbnailInfo = ThumbnailInfo;
+    }
+
+    return ThumbnailInfo;
+    */
+
+    return nullptr;
+}
+
+
+bool FHoudiniCSVAssetActions::CanExecuteReimport(const TArray<TWeakObjectPtr<UHoudiniCSV>> Objects) const
+{
+    for (auto ObjIt = Objects.CreateConstIterator(); ObjIt; ++ObjIt)
+    {
+	auto Object = (*ObjIt).Get();
+	if (!Object)
+	    return false;
+    }
+
+    return true;
+}
+
+void FHoudiniCSVAssetActions::ExecuteReimport(const TArray<TWeakObjectPtr<UHoudiniCSV>> Objects) const
+{
+    for ( auto ObjIt = Objects.CreateConstIterator(); ObjIt; ++ObjIt )
+    {
+	auto Object = (*ObjIt).Get();
+	if ( Object )
+	{
+	    // Fonts fail to reimport if they ask for a new file if missing
+	    FReimportManager::Instance()->Reimport(Object, true);
+	}
+    }
+}
+
+
+bool 
+FHoudiniCSVAssetActions::CanExecuteOpenInEditor(const TArray<TWeakObjectPtr<UHoudiniCSV>> Objects) const
+{
+    for (auto ObjIt = Objects.CreateConstIterator(); ObjIt; ++ObjIt)
+    {
+	auto Object = ( *ObjIt ).Get();
+	if ( !Object )
+	    return false;
+    }
+
+    return true;
+}
+
+void
+FHoudiniCSVAssetActions::ExecuteOpenInEditor(const TArray<TWeakObjectPtr<UHoudiniCSV>> Objects) const
+{
+    for ( auto ObjIt = Objects.CreateConstIterator(); ObjIt; ++ObjIt )
+    {
+	auto Object = ( *ObjIt ).Get();
+	if ( Object )
+	{
+	    // Fonts fail to reimport if they ask for a new file if missing
+	    FPlatformProcess::LaunchFileInDefaultExternalApplication(*(Object->FileName), NULL, ELaunchVerb::Open);
+	}
+    }
+}
+
+#undef LOCTEXT_NAMESPACE
\ No newline at end of file
diff --git a/Engine/Plugins/FX/HoudiniNiagara/Source/HoudiniNiagaraEditor/Private/HoudiniCSVFactory.cpp b/Engine/Plugins/FX/HoudiniNiagara/Source/HoudiniNiagaraEditor/Private/HoudiniCSVFactory.cpp
new file mode 100644
--- /dev/null
+++ b/Engine/Plugins/FX/HoudiniNiagara/Source/HoudiniNiagaraEditor/Private/HoudiniCSVFactory.cpp
@@ -0,0 +1,199 @@
+/*
+* Copyright (c) <2017> Side Effects Software Inc.
+*
+* Permission is hereby granted, free of charge, to any person obtaining a copy
+* of this software and associated documentation files (the "Software"), to deal
+* in the Software without restriction, including without limitation the rights
+* to use, copy, modify, merge, publish, distribute, sublicense, and/or sell
+* copies of the Software, and to permit persons to whom the Software is
+* furnished to do so, subject to the following conditions:
+*
+* The above copyright notice and this permission notice shall be included in all
+* copies or substantial portions of the Software.
+*
+* THE SOFTWARE IS PROVIDED "AS IS", WITHOUT WARRANTY OF ANY KIND, EXPRESS OR
+* IMPLIED, INCLUDING BUT NOT LIMITED TO THE WARRANTIES OF MERCHANTABILITY,
+* FITNESS FOR A PARTICULAR PURPOSE AND NONINFRINGEMENT. IN NO EVENT SHALL THE
+* AUTHORS OR COPYRIGHT HOLDERS BE LIABLE FOR ANY CLAIM, DAMAGES OR OTHER
+* LIABILITY, WHETHER IN AN ACTION OF CONTRACT, TORT OR OTHERWISE, ARISING FROM,
+* OUT OF OR IN CONNECTION WITH THE SOFTWARE OR THE USE OR OTHER DEALINGS IN THE
+* SOFTWARE.
+*
+*/
+
+#include "HoudiniCSVFactory.h"
+#include "HoudiniCSV.h"
+#include "EditorFramework/AssetImportData.h"
+#include "HAL/FileManager.h"
+#include "Misc/FileHelper.h"
+#include "Misc/Paths.h"
+#include "Editor.h"
+
+DEFINE_LOG_CATEGORY(LogHoudiniNiagaraEditor);
+
+#define LOCTEXT_NAMESPACE "HoudiniCSVFactory"
+ 
+/////////////////////////////////////////////////////
+// UHoudiniCSVFactory 
+UHoudiniCSVFactory::UHoudiniCSVFactory(const FObjectInitializer& ObjectInitializer) : Super( ObjectInitializer )
+{
+    // This factory is responsible for manufacturing HoudiniEngine assets.
+    SupportedClass = UHoudiniCSV::StaticClass();
+
+    // This factory does not manufacture new objects from scratch.
+    bCreateNew = false;
+
+    // This factory will open the editor for each new object.
+    bEditAfterNew = false;
+
+    // This factory will import objects from files.
+    bEditorImport = true;
+
+    // Factory does not import objects from text.
+    bText = true;
+
+    // Add supported formats.
+    Formats.Add(FString(TEXT("csv;")) + NSLOCTEXT("HoudiniCSVFactory", "FormatCSV", "CSV File").ToString());
+    Formats.Add(FString(TEXT("hcsv;")) + NSLOCTEXT("HoudiniCSVFactory", "FormatHCSV", "HCSV File").ToString());
+}
+
+
+UObject* 
+UHoudiniCSVFactory::FactoryCreateNew(UClass* Class, UObject* InParent, FName Name, EObjectFlags Flags, UObject* Context, FFeedbackContext* Warn)
+{
+    UHoudiniCSV* NewHoudiniCSVObject = NewObject<UHoudiniCSV>(InParent, Class, Name, Flags | RF_Transactional);
+    return NewHoudiniCSVObject;
+}
+
+UObject* 
+UHoudiniCSVFactory::FactoryCreateFile(UClass* InClass, UObject* InParent, FName InName, EObjectFlags Flags, const FString& Filename, const TCHAR* Parms, FFeedbackContext* Warn, bool& bOutOperationCanceled)
+{
+    bOutOperationCanceled = false;
+
+    UHoudiniCSV* NewHoudiniCSVObject = NewObject<UHoudiniCSV>(InParent, InClass, InName, Flags);
+    if ( !NewHoudiniCSVObject )
+	return nullptr;
+
+    if ( !NewHoudiniCSVObject->UpdateFromFile( Filename ) )
+	return nullptr;   
+
+    // Create reimport information.
+    UAssetImportData * AssetImportData = NewHoudiniCSVObject->AssetImportData;
+    if (!AssetImportData)
+    {
+	AssetImportData = NewObject< UAssetImportData >(NewHoudiniCSVObject, UAssetImportData::StaticClass());
+	NewHoudiniCSVObject->AssetImportData = AssetImportData;
+    }
+
+    //NewHoudiniCSVObject->AssetImportData->Update(Filename);
+    AssetImportData->Update(UFactory::GetCurrentFilename());
+
+    // Broadcast notification that the new asset has been imported.
+    FEditorDelegates::OnAssetPostImport.Broadcast(this, NewHoudiniCSVObject);
+
+    return NewHoudiniCSVObject;
+}
+
+FText
+UHoudiniCSVFactory::GetDisplayName() const
+{
+    return LOCTEXT("HoudiniCSVFactoryDescription", "Houdini CSV File");
+}
+
+bool
+UHoudiniCSVFactory::DoesSupportClass(UClass * Class)
+{
+    return Class == SupportedClass;
+}
+bool 
+UHoudiniCSVFactory::FactoryCanImport(const FString& Filename)
+{
+    const FString Extension = FPaths::GetExtension(Filename);
+
+    if ( Extension == TEXT("csv") || Extension == TEXT("hcsv") )
+    {
+	return true;
+    }
+    return false;
+}
+
+bool 
+UHoudiniCSVFactory::CanReimport(UObject* Obj, TArray<FString>& OutFilenames)
+{
+    UHoudiniCSV* HoudiniCSV = Cast<UHoudiniCSV>(Obj);
+    if (HoudiniCSV)
+    {
+	if (HoudiniCSV->AssetImportData)
+	{
+	    HoudiniCSV->AssetImportData->ExtractFilenames(OutFilenames);
+	}
+	else
+	{
+	    OutFilenames.Add(TEXT(""));
+	}
+	return true;
+    }
+    return false;
+}
+
+void 
+UHoudiniCSVFactory::SetReimportPaths(UObject* Obj, const TArray<FString>& NewReimportPaths)
+{
+    UHoudiniCSV* HoudiniCSV = Cast<UHoudiniCSV>(Obj);
+    if (HoudiniCSV 
+	&& HoudiniCSV->AssetImportData
+	&&  ensure(NewReimportPaths.Num() == 1))
+    {
+	HoudiniCSV->AssetImportData->UpdateFilenameOnly(NewReimportPaths[0]);
+    }
+}
+
+EReimportResult::Type 
+UHoudiniCSVFactory::Reimport(UObject* Obj)
+{
+    UHoudiniCSV* HoudiniCSV = Cast<UHoudiniCSV>(Obj);
+    if (!HoudiniCSV || !HoudiniCSV->AssetImportData )
+	return EReimportResult::Failed;
+
+    const FString FilePath = HoudiniCSV->AssetImportData->GetFirstFilename();
+    if (!FilePath.Len())
+	return EReimportResult::Failed;
+
+    UE_LOG(LogHoudiniNiagaraEditor, Log, TEXT("Performing atomic reimport of [%s]"), *FilePath);
+
+    // Ensure that the file provided by the path exists
+    if (IFileManager::Get().FileSize(*FilePath) == INDEX_NONE)
+    {
+	UE_LOG(LogHoudiniNiagaraEditor, Warning, TEXT("Cannot reimport: source file cannot be found."));
+	return EReimportResult::Failed;
+    }
+
+    bool OutCanceled = false;
+    if (ImportObject(HoudiniCSV->GetClass(), HoudiniCSV->GetOuter(), *HoudiniCSV->GetName(), RF_Public | RF_Standalone, FilePath, nullptr, OutCanceled) != nullptr)
+    {
+	UE_LOG(LogHoudiniNiagaraEditor, Log, TEXT("Imported successfully"));
+	// Try to find the outer package so we can dirty it up
+	if (HoudiniCSV->GetOuter())
+	    HoudiniCSV->GetOuter()->MarkPackageDirty();
+	else
+	    HoudiniCSV->MarkPackageDirty();
+    }
+    else if (OutCanceled)
+    {
+	UE_LOG(LogHoudiniNiagaraEditor, Warning, TEXT("-- import canceled"));
+    }
+    else
+    {
+	UE_LOG(LogHoudiniNiagaraEditor, Warning, TEXT("-- import failed"));
+    }
+
+    return EReimportResult::Succeeded;
+}
+
+int32 
+UHoudiniCSVFactory::GetPriority() const
+{
+    return ImportPriority;
+}
+
+#undef LOCTEXT_NAMESPACE
\ No newline at end of file
diff --git a/Engine/Plugins/FX/HoudiniNiagara/Source/HoudiniNiagaraEditor/Private/HoudiniNiagaraEditor.cpp b/Engine/Plugins/FX/HoudiniNiagara/Source/HoudiniNiagaraEditor/Private/HoudiniNiagaraEditor.cpp
new file mode 100644
--- /dev/null
+++ b/Engine/Plugins/FX/HoudiniNiagara/Source/HoudiniNiagaraEditor/Private/HoudiniNiagaraEditor.cpp
@@ -0,0 +1,64 @@
+/*
+* Copyright (c) <2017> Side Effects Software Inc.
+*
+* Permission is hereby granted, free of charge, to any person obtaining a copy
+* of this software and associated documentation files (the "Software"), to deal
+* in the Software without restriction, including without limitation the rights
+* to use, copy, modify, merge, publish, distribute, sublicense, and/or sell
+* copies of the Software, and to permit persons to whom the Software is
+* furnished to do so, subject to the following conditions:
+*
+* The above copyright notice and this permission notice shall be included in all
+* copies or substantial portions of the Software.
+*
+* THE SOFTWARE IS PROVIDED "AS IS", WITHOUT WARRANTY OF ANY KIND, EXPRESS OR
+* IMPLIED, INCLUDING BUT NOT LIMITED TO THE WARRANTIES OF MERCHANTABILITY,
+* FITNESS FOR A PARTICULAR PURPOSE AND NONINFRINGEMENT. IN NO EVENT SHALL THE
+* AUTHORS OR COPYRIGHT HOLDERS BE LIABLE FOR ANY CLAIM, DAMAGES OR OTHER
+* LIABILITY, WHETHER IN AN ACTION OF CONTRACT, TORT OR OTHERWISE, ARISING FROM,
+* OUT OF OR IN CONNECTION WITH THE SOFTWARE OR THE USE OR OTHER DEALINGS IN THE
+* SOFTWARE.
+*
+*/
+
+#include "HoudiniNiagaraEditor.h"
+#include "HoudiniCSVAssetActions.h"
+#include "AssetRegistryModule.h"
+
+#define LOCTEXT_NAMESPACE "FHoudiniNiagaraModule"
+
+void FHoudiniNiagaraEditorModule::StartupModule()
+{
+    // This code will execute after your module is loaded into memory; the exact timing is specified in the .uplugin file per-module
+
+    // Register the Houdini CSV Type Actions
+    IAssetTools& AssetTools = FModuleManager::LoadModuleChecked< FAssetToolsModule >("AssetTools").Get();
+
+    TSharedRef< IAssetTypeActions > HCSVAction = MakeShareable(new FHoudiniCSVAssetActions());
+    AssetTools.RegisterAssetTypeActions(HCSVAction);
+    AssetTypeActions.Add(HCSVAction);
+}
+
+void FHoudiniNiagaraEditorModule::ShutdownModule()
+{
+    // This function may be called during shutdown to clean up your module.  For modules that support dynamic reloading,
+    // we call this function before unloading the module.
+
+
+    // Unregister asset type actions we have previously registered.
+    if (FModuleManager::Get().IsModuleLoaded("AssetTools"))
+    {
+	IAssetTools & AssetTools = FModuleManager::GetModuleChecked< FAssetToolsModule >("AssetTools").Get();
+
+	for (int32 Index = 0; Index < AssetTypeActions.Num(); ++Index)
+	    AssetTools.UnregisterAssetTypeActions(AssetTypeActions[Index].ToSharedRef());
+
+	AssetTypeActions.Empty();
+    }
+}
+
+
+
+#undef LOCTEXT_NAMESPACE
+	
+IMPLEMENT_MODULE(FHoudiniNiagaraEditorModule, HoudiniNiagara)
diff --git a/Engine/Plugins/FX/HoudiniNiagara/Source/HoudiniNiagaraEditor/Public/HoudiniCSVAssetActions.h b/Engine/Plugins/FX/HoudiniNiagara/Source/HoudiniNiagaraEditor/Public/HoudiniCSVAssetActions.h
new file mode 100644
--- /dev/null
+++ b/Engine/Plugins/FX/HoudiniNiagara/Source/HoudiniNiagaraEditor/Public/HoudiniCSVAssetActions.h
@@ -0,0 +1,61 @@
+/*
+* Copyright (c) <2017> Side Effects Software Inc.
+*
+* Permission is hereby granted, free of charge, to any person obtaining a copy
+* of this software and associated documentation files (the "Software"), to deal
+* in the Software without restriction, including without limitation the rights
+* to use, copy, modify, merge, publish, distribute, sublicense, and/or sell
+* copies of the Software, and to permit persons to whom the Software is
+* furnished to do so, subject to the following conditions:
+*
+* The above copyright notice and this permission notice shall be included in all
+* copies or substantial portions of the Software.
+*
+* THE SOFTWARE IS PROVIDED "AS IS", WITHOUT WARRANTY OF ANY KIND, EXPRESS OR
+* IMPLIED, INCLUDING BUT NOT LIMITED TO THE WARRANTIES OF MERCHANTABILITY,
+* FITNESS FOR A PARTICULAR PURPOSE AND NONINFRINGEMENT. IN NO EVENT SHALL THE
+* AUTHORS OR COPYRIGHT HOLDERS BE LIABLE FOR ANY CLAIM, DAMAGES OR OTHER
+* LIABILITY, WHETHER IN AN ACTION OF CONTRACT, TORT OR OTHERWISE, ARISING FROM,
+* OUT OF OR IN CONNECTION WITH THE SOFTWARE OR THE USE OR OTHER DEALINGS IN THE
+* SOFTWARE.
+*
+*/
+
+#pragma once
+
+#include "AssetTypeActions_Base.h"
+#include "Templates/SharedPointer.h"
+
+class UHoudiniCSV;
+
+
+class FHoudiniCSVAssetActions : public FAssetTypeActions_Base
+{
+    public:
+
+	FHoudiniCSVAssetActions();
+
+    public:
+
+	//~ FAssetTypeActions_Base overrides
+	virtual FText GetName() const override;
+	virtual FColor GetTypeColor() const override;
+	virtual UClass* GetSupportedClass() const override;
+	virtual uint32 GetCategories() override;    
+	virtual class UThumbnailInfo* GetThumbnailInfo(UObject* Asset) const override;
+	virtual bool HasActions(const TArray<UObject*>& InObjects) const override;
+	virtual void GetActions(const TArray<UObject*>& InObjects, FMenuBuilder& MenuBuilder) override; 
+
+	virtual bool CanFilter() override;   
+	virtual FText GetAssetDescription(const FAssetData& AssetData) const override;
+
+    private:
+
+	// Reimport actions
+	bool CanExecuteReimport(const TArray<TWeakObjectPtr<UHoudiniCSV>> Objects) const;
+	void ExecuteReimport(const TArray<TWeakObjectPtr<UHoudiniCSV>> Objects) const;
+
+	// Open In Editor Actions
+	bool CanExecuteOpenInEditor(const TArray<TWeakObjectPtr<UHoudiniCSV>> Objects) const;
+	void ExecuteOpenInEditor(const TArray<TWeakObjectPtr<UHoudiniCSV>> Objects) const;
+};
\ No newline at end of file
diff --git a/Engine/Plugins/FX/HoudiniNiagara/Source/HoudiniNiagaraEditor/Public/HoudiniCSVFactory.h b/Engine/Plugins/FX/HoudiniNiagara/Source/HoudiniNiagaraEditor/Public/HoudiniCSVFactory.h
new file mode 100644
--- /dev/null
+++ b/Engine/Plugins/FX/HoudiniNiagara/Source/HoudiniNiagaraEditor/Public/HoudiniCSVFactory.h
@@ -0,0 +1,64 @@
+/*
+* Copyright (c) <2017> Side Effects Software Inc.
+*
+* Permission is hereby granted, free of charge, to any person obtaining a copy
+* of this software and associated documentation files (the "Software"), to deal
+* in the Software without restriction, including without limitation the rights
+* to use, copy, modify, merge, publish, distribute, sublicense, and/or sell
+* copies of the Software, and to permit persons to whom the Software is
+* furnished to do so, subject to the following conditions:
+*
+* The above copyright notice and this permission notice shall be included in all
+* copies or substantial portions of the Software.
+*
+* THE SOFTWARE IS PROVIDED "AS IS", WITHOUT WARRANTY OF ANY KIND, EXPRESS OR
+* IMPLIED, INCLUDING BUT NOT LIMITED TO THE WARRANTIES OF MERCHANTABILITY,
+* FITNESS FOR A PARTICULAR PURPOSE AND NONINFRINGEMENT. IN NO EVENT SHALL THE
+* AUTHORS OR COPYRIGHT HOLDERS BE LIABLE FOR ANY CLAIM, DAMAGES OR OTHER
+* LIABILITY, WHETHER IN AN ACTION OF CONTRACT, TORT OR OTHERWISE, ARISING FROM,
+* OUT OF OR IN CONNECTION WITH THE SOFTWARE OR THE USE OR OTHER DEALINGS IN THE
+* SOFTWARE.
+*
+*/
+
+#include "CoreMinimal.h"
+#include "UObject/ObjectMacros.h"
+#include "UObject/UObjectGlobals.h"
+#include "UObject/Object.h"
+#include "Factories/Factory.h"
+#include "EditorReimportHandler.h"
+#include "HoudiniCSVFactory.generated.h"
+
+DECLARE_LOG_CATEGORY_EXTERN(LogHoudiniNiagaraEditor, All, All);
+
+/** A factory for HoudiniCSV assets. */
+UCLASS(MinimalAPI/*UNREALED_API*/)
+class UHoudiniCSVFactory : public UFactory, public FReimportHandler
+{
+	GENERATED_UCLASS_BODY()
+ 
+	//~ Begin UFactory Interface
+
+	// New file from the menu, uneeded!
+	virtual UObject* FactoryCreateNew( UClass* Class, UObject* InParent, FName Name, EObjectFlags Flags, UObject* Context, FFeedbackContext* Warn ) override;
+
+	// Import file as text
+	//virtual UObject* FactoryCreateText(UClass* InClass, UObject* InParent, FName InName, EObjectFlags Flags, UObject* Context, const TCHAR* Type, const TCHAR*& Buffer, const TCHAR* BufferEnd, FFeedbackContext* Warn) override;
+	virtual UObject* FactoryCreateFile( UClass* InClass, UObject* InParent, FName InName, EObjectFlags Flags, const FString& Filename, const TCHAR* Parms, FFeedbackContext* Warn, bool& bOutOperationCanceled ) override;
+
+
+	virtual FText GetDisplayName() const override;	
+	virtual bool DoesSupportClass( UClass * Class ) override;
+	virtual bool FactoryCanImport( const FString& Filename ) override;
+
+	//~ End UFactory Interface
+
+	//~ Begin FReimportHandler Interface
+
+	virtual bool CanReimport(UObject* Obj, TArray<FString>& OutFilenames) override;
+	virtual void SetReimportPaths(UObject* Obj, const TArray<FString>& NewReimportPaths) override;
+	virtual EReimportResult::Type Reimport(UObject* Obj) override;
+	virtual int32 GetPriority() const override;
+
+	//~ End FReimportHandler Interface
+};
\ No newline at end of file
diff --git a/Engine/Plugins/FX/HoudiniNiagara/Source/HoudiniNiagaraEditor/Public/HoudiniNiagaraEditor.h b/Engine/Plugins/FX/HoudiniNiagara/Source/HoudiniNiagaraEditor/Public/HoudiniNiagaraEditor.h
new file mode 100644
--- /dev/null
+++ b/Engine/Plugins/FX/HoudiniNiagara/Source/HoudiniNiagaraEditor/Public/HoudiniNiagaraEditor.h
@@ -0,0 +1,43 @@
+/*
+* Copyright (c) <2017> Side Effects Software Inc.
+*
+* Permission is hereby granted, free of charge, to any person obtaining a copy
+* of this software and associated documentation files (the "Software"), to deal
+* in the Software without restriction, including without limitation the rights
+* to use, copy, modify, merge, publish, distribute, sublicense, and/or sell
+* copies of the Software, and to permit persons to whom the Software is
+* furnished to do so, subject to the following conditions:
+*
+* The above copyright notice and this permission notice shall be included in all
+* copies or substantial portions of the Software.
+*
+* THE SOFTWARE IS PROVIDED "AS IS", WITHOUT WARRANTY OF ANY KIND, EXPRESS OR
+* IMPLIED, INCLUDING BUT NOT LIMITED TO THE WARRANTIES OF MERCHANTABILITY,
+* FITNESS FOR A PARTICULAR PURPOSE AND NONINFRINGEMENT. IN NO EVENT SHALL THE
+* AUTHORS OR COPYRIGHT HOLDERS BE LIABLE FOR ANY CLAIM, DAMAGES OR OTHER
+* LIABILITY, WHETHER IN AN ACTION OF CONTRACT, TORT OR OTHERWISE, ARISING FROM,
+* OUT OF OR IN CONNECTION WITH THE SOFTWARE OR THE USE OR OTHER DEALINGS IN THE
+* SOFTWARE.
+*
+*/
+
+#pragma once
+
+#include "CoreMinimal.h"
+#include "Modules/ModuleInterface.h"
+
+class IAssetTypeActions;
+
+class FHoudiniNiagaraEditorModule : public IModuleInterface
+{
+    public:
+
+	/** IModuleInterface implementation */
+	virtual void StartupModule() override;
+	virtual void ShutdownModule() override;
+
+    private:
+
+	/** AssetType actions associated with Houdini CSV assets. **/
+	TArray< TSharedPtr< IAssetTypeActions > > AssetTypeActions;
+};
\ No newline at end of file
diff --git a/Engine/Plugins/FX/Niagara/.p4ignore.txt b/Engine/Plugins/FX/Niagara/.p4ignore.txt
--- a/Engine/Plugins/FX/Niagara/.p4ignore.txt
+++ b/Engine/Plugins/FX/Niagara/.p4ignore.txt
@@ -1 +1,2 @@
 *Binaries*
+*Intermediate*
diff --git a/Engine/Plugins/FX/Niagara/Config/BaseNiagara.ini b/Engine/Plugins/FX/Niagara/Config/BaseNiagara.ini
--- a/Engine/Plugins/FX/Niagara/Config/BaseNiagara.ini
+++ b/Engine/Plugins/FX/Niagara/Config/BaseNiagara.ini
@@ -12,13 +12,14 @@
 
 +EnumRedirects=(OldName="ENiagaraScriptUsage",ValueChanges=(("ENiagaraScriptUsage::SpawnScript","ENiagaraScriptUsage::ParticleSpawnScript"), ("ENiagaraScriptUsage::SpawnScriptInterpolated","ENiagaraScriptUsage::ParticleSpawnScriptInterpolated"),("ENiagaraScriptUsage::UpdateScript","ENiagaraScriptUsage::ParticleUpdateScript"), ("ENiagaraScriptUsage::EventScript","ENiagaraScriptUsage::ParticleEventScript"), ("ENiagaraScriptUsage::SystemScript","ENiagaraScriptUsage::SystemSpawnScript")) )
 
-+PropertyRedirects=(OldName="NiagaraEmitter.ChangeId", NewName="NiagaraEmitter.GraphChangeId")
+;+EnumRedirects=(OldName="ENiagaraSortMode",ValueChanges=(("ENiagaraSortMode::SortNone","ENiagaraSortMode::None"), ("ENiagaraSortMode::SortViewDepth","ENiagaraSortMode::ViewDepth"), ("ENiagaraSortMode::SortViewDistance","ENiagaraSortMode::ViewDistance")) )
 
 [/Script/Niagara.NiagaraSettings]
 +AdditionalParameterEnums=/Niagara/Functions/Localspace/ENiagaraCoordinateSpace.ENiagaraCoordinateSpace
 +AdditionalParameterEnums=/Niagara/Functions/Localspace/ENiagaraOrientationAxis.ENiagaraOrientationAxis
++AdditionalParameterEnums=/Niagara/DynamicInputs/Bool/ENiagaraBooleanLogicOps.ENiagaraBooleanLogicOps
 
-[/Script/NiagaraEditor.NiagaraEditorSettings]
+[/Script/NiagaraEditor.NiagaraEditorSettings] 
 DefaultSystem=None
 DefaultEmitter=/Niagara/DefaultAssets/DefaultEmitter.DefaultEmitter
 DefaultScript=/Niagara/DefaultAssets/DefaultModule.DefaultModule
diff --git a/Engine/Plugins/FX/Niagara/Niagara.uplugin b/Engine/Plugins/FX/Niagara/Niagara.uplugin
--- a/Engine/Plugins/FX/Niagara/Niagara.uplugin
+++ b/Engine/Plugins/FX/Niagara/Niagara.uplugin
@@ -16,36 +16,36 @@
 	"Installed" : false,
 	"Modules" :
 	[
+		{
+		  "Name": "NiagaraCore",
+		  "Type": "Runtime",
+		  "LoadingPhase": "PreDefault"
+		},
 		{
 		  "Name": "Niagara",
 		  "Type": "Runtime",
-		  "LoadingPhase": "PreDefault",
-		  "WhitelistPlatforms" : [ "Win64", "Win32", "PS4", "XboxOne", "Mac" ]
+		  "LoadingPhase": "PreDefault"
 		},
 		{
 			"Name" : "NiagaraShader",
 			"Type" : "Runtime",
-			"LoadingPhase" : "PostConfigInit",
-			"WhitelistPlatforms" : [ "Win64", "Win32", "PS4", "XboxOne", "Mac" ]
+			"LoadingPhase" : "PostConfigInit"
 		},
 
 		{
 			"Name" : "NiagaraVertexFactories",
 			"Type" : "Runtime",
-			"LoadingPhase" : "PostConfigInit",
-			"WhitelistPlatforms" : [ "Win64", "Win32", "PS4", "XboxOne", "Mac" ]
+			"LoadingPhase" : "PostConfigInit"
 		},
 		{
 			"Name" : "NiagaraEditor",
 			"Type" : "Editor",
-			"LoadingPhase" : "Default",
-			"WhitelistPlatforms" : [ "Win64", "Win32", "PS4", "XboxOne", "Mac" ]
+			"LoadingPhase" : "Default"
 		},
 		{
 			"Name" : "NiagaraEditorWidgets",
 			"Type" : "Editor",
-			"LoadingPhase" : "Default",
-			"WhitelistPlatforms" : [ "Win64", "Win32", "PS4", "XboxOne", "Mac" ]
+			"LoadingPhase" : "Default"
 		}
 	]
 }
diff --git a/Engine/Plugins/FX/Niagara/Source/Niagara/Classes/NiagaraCollision.h b/Engine/Plugins/FX/Niagara/Source/Niagara/Classes/NiagaraCollision.h
--- a/Engine/Plugins/FX/Niagara/Source/Niagara/Classes/NiagaraCollision.h
+++ b/Engine/Plugins/FX/Niagara/Source/Niagara/Classes/NiagaraCollision.h
@@ -19,7 +19,8 @@ struct FNiagaraCollisionTrace
 {
 	FTraceHandle CollisionTraceHandle;
 	uint32 SourceParticleIndex;
-	FVector OriginalVelocity;
+	float CollisionSize;
+	float DeltaSeconds;
 };
 
 
@@ -37,11 +38,11 @@ public:
 		FNiagaraEventDataSetMgr::Reset(OwnerSystemInstanceName, EmitterName);
 	}
 
-	void Tick(ENiagaraSimTarget Target)
+	void Tick()
 	{
 		if (CollisionEventDataSet)
 		{
-			CollisionEventDataSet->Tick(Target);
+			CollisionEventDataSet->Tick();
 		}
 	}
 
@@ -57,7 +58,7 @@ public:
 	{
 		if (CollisionEventDataSet)
 		{
-			CollisionEventDataSet->Reset();
+			CollisionEventDataSet->Init(FNiagaraDataSetID(), ENiagaraSimTarget::CPUSim);
 		}
 		EmitterName = InEmitterName;
 		OwnerSystemInstanceName = InOwnerSystemInstanceName;
@@ -155,7 +156,7 @@ public:
 		TraceID = 0;
 	}
 
-	int32 SubmitQuery(FVector Position, FVector Velocity, float CollisionSize, float DeltaSeconds);
+	int32 SubmitQuery(FVector Position, FVector Direction, float CollisionSize, float DeltaSeconds);
 	bool GetQueryResult(uint32 TraceID, FNiagaraDICollsionQueryResult &Result);
 
 private:
