Syncing with fixing in Dev-Niagara up to CL 4118934

Making it so that you can only AddParameters in graphs where it makes sense to do so. Promoting variables also follows the same logic.

#tests made dynamic input, function, and module graphs and made sure that each one was only capable of creating the types that were proper

Fix for community issue where you could connect an input and output node and encounter a check that the warnings variable array was no longer the same size as the variables array.

#tests auto-tests pass as well as jira use case

Houdini - Niagara:
Added a version to FNiagaraDataInterfaceParametersCS_HoudiniCSV to avoid further issue when deserializing them.
Fixed bug causing no points to spawn when all points in a file spawned at the same time value.
Removed unneeded tests/restriction when calling functions requiring time value on file without time values.
When no time value is in the file, it is cosidered to always be 0.
Fixed differences between GPU/CPU code.
Fixed special attributes column for Vectors in the CSV file not properly named after being expanded.

Houdini - Niagara:
Moved the logic of GetPointIDsToSpawnAtTime to HoudiniCSV  so the function can be called from blueprints.
Renamed the MinIndex/MaxIndex parameters on it to MinId/MaxID.

Multiple fixes to the spline component, which hasn't received attention in a long time. It is only supported on CPU VM.

#tests new auto-test now passes

System Sim for Update On Spawn

Fix for crash where dynamic inputs didn't have a proper parameter map input value.

#tests new auto-test added

Houdini-Niagara:
Updated the uplugin file.

diff --git a/Engine/Plugins/FX/HoudiniNiagara/HoudiniNiagara.uplugin b/Engine/Plugins/FX/HoudiniNiagara/HoudiniNiagara.uplugin
--- a/Engine/Plugins/FX/HoudiniNiagara/HoudiniNiagara.uplugin
+++ b/Engine/Plugins/FX/HoudiniNiagara/HoudiniNiagara.uplugin
@@ -3,19 +3,18 @@
 	"Version": 1,
 	"VersionName": "1.0",
 	"FriendlyName": "Houdini Niagara",
-	"Description": "Houdini Niagara Integration",
+	"Description": "Import Houdini point cloud data into UE4 to drive Niagara particle systems.",
 	"Category": "FX",
 	"CreatedBy": "Side Effects Software Inc.",
 	"CreatedByURL": "http://www.sidefx.com",
-	"DocsURL": "",
+	"DocsURL": "https://www.sidefx.com/docs/unreal/_niagara.html",
 	"MarketplaceURL": "",
 	"SupportURL": "",
 	"EnabledByDefault": false,
 	"CanContainContent": true,
 	"IsBetaVersion": true,
 	"Installed": false,
-	"Modules":
-	[
+	"Modules": [
 		{
 			"Name": "HoudiniNiagara",
 			"Type": "Runtime",
@@ -30,8 +29,7 @@
 		}
 	],
 
-	"Plugins":
-	[
+	"Plugins": [
 		{
 			"Name": "Niagara",
 			"Enabled": true
diff --git a/Engine/Plugins/FX/HoudiniNiagara/Source/HoudiniNiagara/Classes/NiagaraDataInterfaceHoudiniCSV.h b/Engine/Plugins/FX/HoudiniNiagara/Source/HoudiniNiagara/Classes/NiagaraDataInterfaceHoudiniCSV.h
--- a/Engine/Plugins/FX/HoudiniNiagara/Source/HoudiniNiagara/Classes/NiagaraDataInterfaceHoudiniCSV.h
+++ b/Engine/Plugins/FX/HoudiniNiagara/Source/HoudiniNiagara/Classes/NiagaraDataInterfaceHoudiniCSV.h
@@ -37,23 +37,27 @@ UCLASS(EditInlineNew, Category = "Houdini Niagara", meta = (DisplayName = "Houdi
 class HOUDININIAGARA_API UNiagaraDataInterfaceHoudiniCSV : public UNiagaraDataInterface
 {
 	GENERATED_UCLASS_BODY()
+
 public:
 
+	// Houdini CSV Asset to sample
 	UPROPERTY( EditAnywhere, Category = "Houdini Niagara", meta = (DisplayName = "Houdini CSV Asset" ) )
 	UHoudiniCSV* HoudiniCSVAsset;
 
 	//----------------------------------------------------------------------------
-	//UObject Interface
+	// UObject Interface
 	virtual void PostInitProperties() override;
 	virtual void PostLoad() override;
 #if WITH_EDITOR
 	virtual void PostEditChangeProperty(struct FPropertyChangedEvent& PropertyChangedEvent) override;
 #endif
 	//----------------------------------------------------------------------------
+	// UNiagaraDataInterface Interface
 
+	// Returns the DI's functions signatures 
 	virtual void GetFunctions(TArray<FNiagaraFunctionSignature>& OutFunctions)override;
 	
-	/** Returns the delegate for the passed function signature. */
+	// Returns the delegate for the passed function signature.
 	virtual void GetVMExternalFunction(const FVMExternalFunctionBindingInfo& BindingInfo, void* InstanceData, FVMExternalFunction &OutFunc) override;
 
 	virtual bool Equals(const UNiagaraDataInterface* Other) const override;
@@ -155,9 +159,47 @@ public:
 	// GPU / HLSL Functions
 	virtual bool GetFunctionHLSL(const FName& DefinitionFunctionName, FString InstanceFunctionName, FNiagaraDataInterfaceGPUParamInfo& ParamInfo, FString& OutHLSL) override;
 	virtual void GetParameterDefinitionHLSL(FNiagaraDataInterfaceGPUParamInfo& ParamInfo, FString& OutHLSL) override;
-
-	// Disabling GPU sim for now.
-	virtual bool CanExecuteOnTarget(ENiagaraSimTarget Target)const override { return Target == ENiagaraSimTarget::CPUSim; }
+	virtual FNiagaraDataInterfaceParametersCS* ConstructComputeParameters() const override;
+
+	virtual bool CanExecuteOnTarget(ENiagaraSimTarget Target)const override { return true; }
+
+	// Members for GPU compatibility
+
+	// Buffer and member variables base name
+	static const FString NumberOfRowsBaseName;
+	static const FString NumberOfColumnsBaseName;
+	static const FString NumberOfPointsBaseName;
+	static const FString FloatValuesBufferBaseName;
+	static const FString SpecialAttributesColumnIndexesBufferBaseName;
+	static const FString SpawnTimesBufferBaseName;
+	static const FString LifeValuesBufferBaseName;
+	static const FString PointTypesBufferBaseName;
+	static const FString MaxNumberOfIndexesPerPointBaseName;
+	static const FString PointValueIndexesBufferBaseName;
+	static const FString LastSpawnedPointIdBaseName;
+	static const FString LastSpawnTimeBaseName;
+
+	// Member variables accessors
+	FORCEINLINE int32 GetNumberOfRows()const { return HoudiniCSVAsset ? HoudiniCSVAsset->GetNumberOfRows() : 0; }
+	FORCEINLINE int32 GetNumberOfColumns()const { return HoudiniCSVAsset ? HoudiniCSVAsset->GetNumberOfColumns() : 0; }
+	FORCEINLINE int32 GetNumberOfPoints()const { return HoudiniCSVAsset ? HoudiniCSVAsset->GetNumberOfPoints() : 0; }
+	FORCEINLINE int32 GetMaxNumberOfIndexesPerPoints()const { return HoudiniCSVAsset ? HoudiniCSVAsset->GetMaxNumberOfPointValueIndexes() + 1 : 0; }
+
+	// GPU Buffers accessors
+	FRWBuffer& GetFloatValuesGPUBuffer();
+	FRWBuffer& GetSpecialAttributesColumnIndexesGPUBuffer();
+	FRWBuffer& GetSpawnTimesGPUBuffer();
+	FRWBuffer& GetLifeValuesGPUBuffer();
+	FRWBuffer& GetPointTypesGPUBuffer();
+	FRWBuffer& GetPointValueIndexesGPUBuffer();
+
+	// GPU Buffers
+	FRWBuffer FloatValuesGPUBuffer;
+	FRWBuffer SpecialAttributesColumnIndexesGPUBuffer;
+	FRWBuffer SpawnTimesGPUBuffer;
+	FRWBuffer LifeValuesGPUBuffer;
+	FRWBuffer PointTypesGPUBuffer;
+	FRWBuffer PointValueIndexesGPUBuffer;
 
 protected:
 
@@ -165,9 +207,22 @@ protected:
 
 	// Indicates the GPU buffers need to be updated
 	UPROPERTY()
-	bool GPUBufferDirty;
+	bool FloatValuesGPUBufferDirty;
+	UPROPERTY()
+	bool SpecialAttributesColumnIndexesGPUBufferDirty;
+	UPROPERTY()
+	bool SpawnTimesGPUBufferDirty;
+	UPROPERTY()
+	bool LifeValuesGPUBufferDirty;
+	UPROPERTY()
+	bool PointTypesGPUBufferDirty;
+	UPROPERTY()
+	bool PointValueIndexesGPUBufferDirty;
 
 	// Last Spawned PointID
 	UPROPERTY()
 	int32 LastSpawnedPointID;
+	// Last Spawn time
+	UPROPERTY()
+	float LastSpawnTime;
 };
diff --git a/Engine/Plugins/FX/HoudiniNiagara/Source/HoudiniNiagara/Private/HoudiniCSV.cpp b/Engine/Plugins/FX/HoudiniNiagara/Source/HoudiniNiagara/Private/HoudiniCSV.cpp
--- a/Engine/Plugins/FX/HoudiniNiagara/Source/HoudiniNiagara/Private/HoudiniCSV.cpp
+++ b/Engine/Plugins/FX/HoudiniNiagara/Source/HoudiniNiagara/Private/HoudiniCSV.cpp
@@ -34,7 +34,6 @@
 
 DEFINE_LOG_CATEGORY(LogHoudiniNiagara);
 
-
 struct FHoudiniCSVSortPredicate
 {
 	FHoudiniCSVSortPredicate(const int32 &InTimeCol, const int32 &InIDCol )
@@ -310,9 +309,13 @@ bool UHoudiniCSV::UpdateFromStringArray( TArray<FString>& RawStringArray )
 			*/
 		}
 
-		// Look at the point ID, the max ID will be our number of points
-		if ( NumberOfPoints <= CurrentID )
-			NumberOfPoints = CurrentID + 1;
+		// If we dont have Point ID informations, we still want to fill the PointValueIndexes array
+		if ( !IsValidAttributeColumnIndex( EHoudiniAttributes::POINTID ) )
+		{
+			// Each row is considered its own point
+			PointValueIndexes.Add(FPointIndexes());
+			PointValueIndexes[ rowIdx ].RowIndexes.Add( rowIdx );
+		}
     }
 	
 	NumberOfPoints = HoudiniIDToNiagaraIDMap.Num();
@@ -472,18 +475,6 @@ bool UHoudiniCSV::ParseCSVTitleRow( const FString& TitleRow, const FString& Firs
 		}
     }
 
-	// Extract the special attributes column indexes we found
-	int32 PositionColumnIndex = GetAttributeColumnIndex(EHoudiniAttributes::POSITION);
-	int32 NormalColumnIndex = GetAttributeColumnIndex(EHoudiniAttributes::NORMAL);
-	int32 TimeColumnIndex = GetAttributeColumnIndex(EHoudiniAttributes::TIME);
-	int32 IDColumnIndex = GetAttributeColumnIndex(EHoudiniAttributes::POINTID);
-	int32 AliveColumnIndex = GetAttributeColumnIndex(EHoudiniAttributes::ALIVE);
-	int32 LifeColumnIndex = GetAttributeColumnIndex(EHoudiniAttributes::LIFE);
-	int32 ColorColumnIndex = GetAttributeColumnIndex(EHoudiniAttributes::COLOR);
-	int32 AlphaColumnIndex = GetAttributeColumnIndex(EHoudiniAttributes::ALPHA);
-	int32 VelocityColumnIndex = GetAttributeColumnIndex(EHoudiniAttributes::VELOCITY);
-	int32 TypeColumnIndex = GetAttributeColumnIndex(EHoudiniAttributes::TYPE);
-
 	// Read the first row of the CSV file, and look for packed vectors value (X,Y,Z)
     // We'll have to expand them in the title row to match the parsed data
 	HasPackedVectors = false;
@@ -530,6 +521,18 @@ bool UHoudiniCSV::ParseCSVTitleRow( const FString& TitleRow, const FString& Firs
 		// Increment the number of columns
 		NumberOfColumns += ( FoundVectorSize - 1 );
 
+		// Extract the special attributes column indexes we found
+		int32 PositionColumnIndex = GetAttributeColumnIndex(EHoudiniAttributes::POSITION);
+		int32 NormalColumnIndex = GetAttributeColumnIndex(EHoudiniAttributes::NORMAL);
+		int32 TimeColumnIndex = GetAttributeColumnIndex(EHoudiniAttributes::TIME);
+		int32 IDColumnIndex = GetAttributeColumnIndex(EHoudiniAttributes::POINTID);
+		int32 AliveColumnIndex = GetAttributeColumnIndex(EHoudiniAttributes::ALIVE);
+		int32 LifeColumnIndex = GetAttributeColumnIndex(EHoudiniAttributes::LIFE);
+		int32 ColorColumnIndex = GetAttributeColumnIndex(EHoudiniAttributes::COLOR);
+		int32 AlphaColumnIndex = GetAttributeColumnIndex(EHoudiniAttributes::ALPHA);
+		int32 VelocityColumnIndex = GetAttributeColumnIndex(EHoudiniAttributes::VELOCITY);
+		int32 TypeColumnIndex = GetAttributeColumnIndex(EHoudiniAttributes::TYPE);
+
 		// Expand TitleRowArray
 		if ( ( FoundPackedVectorColumnIndex == PositionColumnIndex ) && ( FoundVectorSize == 3 ) )
 		{
@@ -833,6 +836,7 @@ bool UHoudiniCSV::GetLastRowIndexAtTime(const float& desiredTime, int32& lastRow
 // If desiredTime is higher than the last point time in the csv file, LastIndex will be set to the last point's index
 bool UHoudiniCSV::GetLastPointIDToSpawnAtTime( const float& desiredTime, int32& lastID ) const
 {
+	/*
 	// If we dont have proper time info, always return the last point
 	int32 TimeColumnIndex = GetAttributeColumnIndex(EHoudiniAttributes::TIME);    
     if ( TimeColumnIndex < 0 || TimeColumnIndex >= NumberOfColumns )
@@ -840,6 +844,7 @@ bool UHoudiniCSV::GetLastPointIDToSpawnAtTime( const float& desiredTime, int32&
 		lastID = NumberOfPoints - 1;
 		return false;
     }
+	*/
 
     float temp_time = 0.0f;
 	if ( !SpawnTimes.IsValidIndex( NumberOfPoints - 1 ) )
@@ -972,6 +977,18 @@ bool UHoudiniCSV::GetPointVectorValueAtTime( const int32& PointID, const int32&
 	return true;
 }
 
+int32 UHoudiniCSV::GetMaxNumberOfPointValueIndexes() const
+{
+	int32 MaxNum = 0;
+	for ( auto ValueIndexes : PointValueIndexes )
+	{
+		if ( MaxNum < ValueIndexes.RowIndexes.Num() )
+			MaxNum = ValueIndexes.RowIndexes.Num();
+	}
+
+	return MaxNum;
+}
+
 bool UHoudiniCSV::GetRowIndexesForPointAtTime(const int32& PointID, const float& desiredTime, int32& PrevIndex, int32& NextIndex, float& PrevWeight ) const
 {
 	float PrevTime = -1.0f;
@@ -981,33 +998,6 @@ bool UHoudiniCSV::GetRowIndexesForPointAtTime(const int32& PointID, const float&
 	if ( PointID < 0 || PointID >= NumberOfPoints )
 		return false;
 
-	// Get the spawn time
-	float PointSpawnTime = 0.0f;
-	if ( SpawnTimes.IsValidIndex( PointID ) )
-		PointSpawnTime = SpawnTimes[ PointID ];
-
-	// If the point hasn't spawn before DesiredTime
-	if ( PointSpawnTime > desiredTime )
-		return false;
-
-	// Get the life value
-	float PointLife = 0.0f;
-	if ( LifeValues.IsValidIndex( PointID ) )
-		PointLife = LifeValues[ PointID ];
-
-	// If the point is dead before DesiredTime
-	if ( PointLife > 0.0f && ( PointSpawnTime + PointLife < desiredTime ) )	
-		return false;
-
-	// We don't have Id information
-	// return PointID for the rowIndexes ??
-	if (!IsValidAttributeColumnIndex(EHoudiniAttributes::POINTID))
-		return false;
-
-	// We don't have time information
-	if (!IsValidAttributeColumnIndex(EHoudiniAttributes::TIME))
-		return false;
-
 	// Get the row indexes for this point
 	const TArray<int32>* RowIndexes = nullptr;
 	if ( PointValueIndexes.IsValidIndex( PointID ) )
@@ -1020,33 +1010,33 @@ bool UHoudiniCSV::GetRowIndexesForPointAtTime(const int32& PointID, const float&
 	{
 		// Get the time
 		float currentTime = -1.0f;
-		if ( GetTimeValue( n, currentTime) )
+		if ( !GetTimeValue(n, currentTime) )
+			currentTime = 0.0f;
+
+		if ( currentTime == desiredTime )
+		{
+			PrevIndex = n;
+			NextIndex = n;
+			PrevWeight = 1.0f;
+			return true;
+		}
+		else if ( currentTime < desiredTime )
 		{
-			if ( currentTime == desiredTime )
+			if ( PrevTime == -1.0f || PrevTime < currentTime )
 			{
 				PrevIndex = n;
-				NextIndex = n;
-				PrevWeight = 1.0f;
-				return true;
-			}
-			else if ( currentTime < desiredTime )
-			{
-				if ( PrevTime == -1.0f || PrevTime < currentTime )
-				{
-					PrevIndex = n;
-					PrevTime = currentTime;
-				}
+				PrevTime = currentTime;
 			}
-			else
+		}
+		else
+		{
+			if ( NextTime == -1.0f || NextTime > currentTime)
 			{
-				if ( NextTime == -1.0f || NextTime > currentTime)
-				{
-					NextIndex = n;
-					NextTime = currentTime;
+				NextIndex = n;
+				NextTime = currentTime;
 
-					// TODO: since the csv is sorted by time, we can break now
-					break;
-				}
+				// TODO: since the csv is sorted by time, we can break now
+				break;
 			}
 		}
 	}
@@ -1074,6 +1064,60 @@ bool UHoudiniCSV::GetRowIndexesForPointAtTime(const int32& PointID, const float&
 	return true;
 }
 
+bool UHoudiniCSV::GetPointIDsToSpawnAtTime(
+	const float& desiredTime,
+	int32& MinID, int32& MaxID, int32& Count,
+	int32& LastSpawnedPointID, float& LastSpawnTime ) const
+{
+	int32 lastID = 0;
+	if ( !GetLastPointIDToSpawnAtTime( desiredTime, lastID ) )
+	{
+		// The CSV file doesn't have time informations, so always return all points in the file
+		MinID = 0;
+		MaxID = lastID;
+		Count = MaxID - MinID + 1;
+	}
+	else
+	{
+		// The CSV file has time informations
+		// First, detect if we need to reset LastSpawnedPointID (after a loop of the emitter)
+		if ( lastID < LastSpawnedPointID || desiredTime <= LastSpawnTime )
+			LastSpawnedPointID = -1;
+
+		if ( lastID < 0 )
+		{
+			// Nothing to spawn, t is lower than the point's time
+			LastSpawnedPointID = -1;					
+		}
+		else
+		{
+			// The last time value in the CSV is lower than t, spawn everything if we didnt already!
+			if ( lastID >= GetNumberOfPoints() )
+				lastID = lastID - 1;
+
+			if ( lastID == LastSpawnedPointID )
+			{
+				// We dont have any new point to spawn
+				MinID = lastID;
+				MaxID = lastID;
+				Count = 0;
+			}
+			else
+			{
+				// We have points to spawn at time t
+				MinID = LastSpawnedPointID + 1;
+				MaxID = lastID;
+				Count = MaxID - MinID + 1;
+
+				LastSpawnedPointID = MaxID;
+				LastSpawnTime = desiredTime;
+			}
+		}
+	}
+
+	return true;
+}
+
 int32 UHoudiniCSV::GetAttributeColumnIndex(const EHoudiniAttributes& Attr) const
 {
 	if ( !IsValidAttributeColumnIndex( Attr ) )
@@ -1098,21 +1142,28 @@ bool UHoudiniCSV::IsValidAttributeColumnIndex(const EHoudiniAttributes& Attr) co
 // Returns the column index for a given string
 bool UHoudiniCSV::GetColumnIndexFromString( const FString& ColumnTitle, int32& ColumnIndex ) const
 {
-    if ( !ColumnTitleArray.Find( ColumnTitle, ColumnIndex ) )
+    if ( ColumnTitleArray.Find( ColumnTitle, ColumnIndex ) )
 		return true;
 
-    // Handle packed positions/normals here
-    if ( ColumnTitle.Equals( "P" ) )
+    // Handling special attributes and packed positions/normals
+    if ( ColumnTitle.Equals( "P", ESearchCase::IgnoreCase ) )
 		return ColumnTitleArray.Find( TEXT( "Px" ), ColumnIndex );
-    else if ( ColumnTitle.Equals( "N" ) )
+    else if ( ColumnTitle.Equals( "N", ESearchCase::IgnoreCase ) )
 		return ColumnTitleArray.Find( TEXT( "Nx" ), ColumnIndex );
-	// TODO COLOR etc...
+	else if ( ColumnTitle.Equals("V", ESearchCase::IgnoreCase ) )
+		return ColumnTitleArray.Find( TEXT("Vx"), ColumnIndex);
+	else if ( ColumnTitle.Equals("Cd", ESearchCase::IgnoreCase ) )
+		return ColumnTitleArray.Find( TEXT("R"), ColumnIndex);
+	else if ( ColumnTitle.Equals("Alpha", ESearchCase::IgnoreCase ) )
+		return ColumnTitleArray.Find( TEXT("A"), ColumnIndex);
+	else
+		return ColumnTitleArray.Find( ColumnTitle + TEXT("1"), ColumnIndex );
 
     return false;
 }
 
 // Returns the float value at a given point in the CSV file
-bool UHoudiniCSV::GetFloatValue( const int32& rowIndex, const FString& ColumnTitle, float& value ) const
+bool UHoudiniCSV::GetFloatValueForString( const int32& rowIndex, const FString& ColumnTitle, float& value ) const
 {
     int32 ColIndex = -1;
     if ( !GetColumnIndexFromString( ColumnTitle, ColIndex ) )
diff --git a/Engine/Plugins/FX/HoudiniNiagara/Source/HoudiniNiagara/Private/NiagaraDataInterfaceHoudiniCSV.cpp b/Engine/Plugins/FX/HoudiniNiagara/Source/HoudiniNiagara/Private/NiagaraDataInterfaceHoudiniCSV.cpp
--- a/Engine/Plugins/FX/HoudiniNiagara/Source/HoudiniNiagara/Private/NiagaraDataInterfaceHoudiniCSV.cpp
+++ b/Engine/Plugins/FX/HoudiniNiagara/Source/HoudiniNiagara/Private/NiagaraDataInterfaceHoudiniCSV.cpp
@@ -24,8 +24,55 @@
 #include "NiagaraDataInterfaceHoudiniCSV.h"
 #include "NiagaraTypes.h"
 #include "Misc/FileHelper.h"
-
-#define LOCTEXT_NAMESPACE "HoudiniNiagaraCSVDataInterface"  
+#include "NiagaraShader.h"
+#include "ShaderParameterUtils.h"
+
+#define LOCTEXT_NAMESPACE "HoudiniNiagaraCSVDataInterface"
+
+// Base name for the member variables / buffers for GPU compatibility
+const FString UNiagaraDataInterfaceHoudiniCSV::NumberOfRowsBaseName(TEXT("NumberOfRows_"));
+const FString UNiagaraDataInterfaceHoudiniCSV::NumberOfColumnsBaseName(TEXT("NumberOfColumns_"));
+const FString UNiagaraDataInterfaceHoudiniCSV::NumberOfPointsBaseName(TEXT("NumberOfPoints_"));
+const FString UNiagaraDataInterfaceHoudiniCSV::FloatValuesBufferBaseName(TEXT("FloatValuesBuffer_"));
+const FString UNiagaraDataInterfaceHoudiniCSV::SpecialAttributesColumnIndexesBufferBaseName(TEXT("SpecialAttributesColumnIndexesBuffer_"));
+const FString UNiagaraDataInterfaceHoudiniCSV::SpawnTimesBufferBaseName(TEXT("SpawnTimesBuffer_"));
+const FString UNiagaraDataInterfaceHoudiniCSV::LifeValuesBufferBaseName(TEXT("LifeValuesBuffer_"));
+const FString UNiagaraDataInterfaceHoudiniCSV::PointTypesBufferBaseName(TEXT("PointTypesBuffer_"));
+const FString UNiagaraDataInterfaceHoudiniCSV::MaxNumberOfIndexesPerPointBaseName(TEXT("MaxNumberOfIndexesPerPoint_"));
+const FString UNiagaraDataInterfaceHoudiniCSV::PointValueIndexesBufferBaseName(TEXT("PointValueIndexesBuffer_"));
+const FString UNiagaraDataInterfaceHoudiniCSV::LastSpawnedPointIdBaseName(TEXT("LastSpawnedPointId_"));
+const FString UNiagaraDataInterfaceHoudiniCSV::LastSpawnTimeBaseName(TEXT("LastSpawnTime_"));
+
+
+// Name of all the functions available in the data interface
+static const FName GetFloatValueName("GetFloatValue");
+//static const FName GetFloatValueByStringName("GetFloatValueByString");
+static const FName GetVectorValueName("GetVectorValue");
+static const FName GetVectorValueExName("GetVectorValueEx");
+
+static const FName GetPositionName("GetPosition");
+static const FName GetNormalName("GetNormal");
+static const FName GetTimeName("GetTime");
+static const FName GetVelocityName("GetVelocity");
+static const FName GetColorName("GetColor");
+static const FName GetPositionAndTimeName("GetPositionAndTime");
+
+static const FName GetNumberOfPointsName("GetNumberOfPoints");
+static const FName GetNumberOfRowsName("GetNumberOfRows");
+static const FName GetNumberOfColumnsName("GetNumberOfColumns");
+
+static const FName GetLastRowIndexAtTimeName("GetLastRowIndexAtTime");
+static const FName GetPointIDsToSpawnAtTimeName("GetPointIDsToSpawnAtTime");
+static const FName GetRowIndexesForPointAtTimeName("GetRowIndexesForPointAtTime");
+
+static const FName GetPointPositionAtTimeName("GetPointPositionAtTime");
+static const FName GetPointValueAtTimeName("GetPointValueAtTime");
+static const FName GetPointVectorValueAtTimeName("GetPointVectorValueAtTime");
+
+static const FName GetPointVectorValueAtTimeExName("GetPointVectorValueAtTimeEx");
+static const FName GetPointLifeName("GetPointLife");
+//static const FName GetPointLifeAtTimeName("GetPointLifeAtTime");
+static const FName GetPointTypeName("GetPointType");
 
 
 UNiagaraDataInterfaceHoudiniCSV::UNiagaraDataInterfaceHoudiniCSV(FObjectInitializer const& ObjectInitializer)
@@ -44,15 +91,30 @@ void UNiagaraDataInterfaceHoudiniCSV::PostInitProperties()
 	    FNiagaraTypeRegistry::Register(FNiagaraTypeDefinition(GetClass()), true, false, false);
     }
 
-    GPUBufferDirty = true;
+	FloatValuesGPUBufferDirty = true;
+	SpecialAttributesColumnIndexesGPUBufferDirty = true;
+	SpawnTimesGPUBufferDirty = true;
+	LifeValuesGPUBufferDirty = true;
+	PointTypesGPUBufferDirty = true;
+	PointValueIndexesGPUBufferDirty = true;
+
 	LastSpawnedPointID = -1;
+	LastSpawnTime = -1.0f;
 }
 
 void UNiagaraDataInterfaceHoudiniCSV::PostLoad()
 {
     Super::PostLoad();
-    GPUBufferDirty = true;
+
+	FloatValuesGPUBufferDirty = true;
+	SpecialAttributesColumnIndexesGPUBufferDirty = true;
+	SpawnTimesGPUBufferDirty = true;
+	LifeValuesGPUBufferDirty = true;
+	PointTypesGPUBufferDirty = true;
+	PointValueIndexesGPUBufferDirty = true;
+
 	LastSpawnedPointID = -1;
+	LastSpawnTime = -1.0f;
 }
 
 #if WITH_EDITOR
@@ -64,10 +126,17 @@ void UNiagaraDataInterfaceHoudiniCSV::PostEditChangeProperty(struct FPropertyCha
     if (PropertyChangedEvent.Property && PropertyChangedEvent.Property->GetFName() == GET_MEMBER_NAME_CHECKED(UNiagaraDataInterfaceHoudiniCSV, HoudiniCSVAsset))
     {
 		Modify();
-		if (HoudiniCSVAsset)
+		if ( HoudiniCSVAsset )
 		{
-			GPUBufferDirty = true;
+			FloatValuesGPUBufferDirty = true;
+			SpecialAttributesColumnIndexesGPUBufferDirty = true;
+			SpawnTimesGPUBufferDirty = true;
+			LifeValuesGPUBufferDirty = true;
+			PointTypesGPUBufferDirty = true;
+			PointValueIndexesGPUBufferDirty = true;
+
 			LastSpawnedPointID = -1;
+			LastSpawnTime = -1.0f;
 		}
     }
 }
@@ -85,6 +154,15 @@ bool UNiagaraDataInterfaceHoudiniCSV::CopyToInternal(UNiagaraDataInterface* Dest
 
     CastedInterface->HoudiniCSVAsset = HoudiniCSVAsset;
 
+	CastedInterface->FloatValuesGPUBufferDirty = true;
+	CastedInterface->SpecialAttributesColumnIndexesGPUBufferDirty = true;
+	CastedInterface->SpawnTimesGPUBufferDirty = true;
+	CastedInterface->LifeValuesGPUBufferDirty = true;
+	CastedInterface->PointTypesGPUBufferDirty = true;
+	CastedInterface->PointValueIndexesGPUBufferDirty = true;
+	CastedInterface->LastSpawnedPointID = -1;
+	CastedInterface->LastSpawnTime = -1.0f;
+
     return true;
 }
 
@@ -104,13 +182,13 @@ bool UNiagaraDataInterfaceHoudiniCSV::Equals(const UNiagaraDataInterface* Other)
     return false;
 }
 
-
+// Returns the signature of all the functions avaialable in the data interface
 void UNiagaraDataInterfaceHoudiniCSV::GetFunctions(TArray<FNiagaraFunctionSignature>& OutFunctions)
 {
     {
 		// GetFloatValue
 		FNiagaraFunctionSignature Sig;
-		Sig.Name = TEXT("GetFloatValue");
+		Sig.Name = GetFloatValueName;
 		Sig.bMemberFunction = true;
 		Sig.bRequiresContext = false;
 		Sig.Inputs.Add(FNiagaraVariable(FNiagaraTypeDefinition(GetClass()), TEXT("CSV")));			// CSV in
@@ -128,11 +206,11 @@ void UNiagaraDataInterfaceHoudiniCSV::GetFunctions(TArray<FNiagaraFunctionSignat
     {
 		// GetFloatValueByString
 		FNiagaraFunctionSignature Sig;
-		Sig.Name = TEXT("GetFloatValueByString");
+		Sig.Name = GetFloatValueByStringName;
 		Sig.bMemberFunction = true;
 		Sig.bRequiresContext = false;
 		Sig.Inputs.Add(FNiagaraVariable(FNiagaraTypeDefinition(GetClass()), TEXT("CSV")));			// CSV in
-		Sig.Inputs.Add(FNiagaraVariable(FNiagaraTypeDefinition::GetIntDef(), TEXT("Row")));		// Row Index In
+		Sig.Inputs.Add(FNiagaraVariable(FNiagaraTypeDefinition::GetIntDef(), TEXT("Row")));			// Row Index In
 		Sig.Inputs.Add(FNiagaraVariable(FNiagaraTypeDefinition::GetStringDef(), TEXT("ColTitle")));	// Col Title In
 		Sig.Outputs.Add(FNiagaraVariable(FNiagaraTypeDefinition::GetFloatDef(), TEXT("Value")));	// Float Out
 
@@ -146,7 +224,7 @@ void UNiagaraDataInterfaceHoudiniCSV::GetFunctions(TArray<FNiagaraFunctionSignat
 	{
 		// GetVectorValue
 		FNiagaraFunctionSignature Sig;
-		Sig.Name = TEXT("GetVectorValue");
+		Sig.Name = GetVectorValueName;
 		Sig.bMemberFunction = true;
 		Sig.bRequiresContext = false;
 		Sig.Inputs.Add(FNiagaraVariable(FNiagaraTypeDefinition(GetClass()), TEXT("CSV")));			// CSV in
@@ -163,7 +241,7 @@ void UNiagaraDataInterfaceHoudiniCSV::GetFunctions(TArray<FNiagaraFunctionSignat
 	{
 		// GetVectorValueEx
 		FNiagaraFunctionSignature Sig;
-		Sig.Name = TEXT("GetVectorValueEx");
+		Sig.Name = GetVectorValueExName;
 		Sig.bMemberFunction = true;
 		Sig.bRequiresContext = false;
 		Sig.Inputs.Add(FNiagaraVariable(FNiagaraTypeDefinition(GetClass()), TEXT("CSV")));			// CSV in
@@ -182,7 +260,7 @@ void UNiagaraDataInterfaceHoudiniCSV::GetFunctions(TArray<FNiagaraFunctionSignat
     {
 		// GetPosition
 		FNiagaraFunctionSignature Sig;
-		Sig.Name = TEXT("GetPosition");
+		Sig.Name = GetPositionName;
 		Sig.bMemberFunction = true;
 		Sig.bRequiresContext = false;
 		Sig.Inputs.Add(FNiagaraVariable(FNiagaraTypeDefinition(GetClass()), TEXT("CSV")));			// CSV in
@@ -198,7 +276,7 @@ void UNiagaraDataInterfaceHoudiniCSV::GetFunctions(TArray<FNiagaraFunctionSignat
     {
 		// GetPositionAndTime
 		FNiagaraFunctionSignature Sig;
-		Sig.Name = TEXT("GetPositionAndTime");
+		Sig.Name = GetPositionAndTimeName;
 		Sig.bMemberFunction = true;
 		Sig.bRequiresContext = false;
 		Sig.Inputs.Add(FNiagaraVariable(FNiagaraTypeDefinition(GetClass()), TEXT("CSV")));			// CSV in
@@ -215,7 +293,7 @@ void UNiagaraDataInterfaceHoudiniCSV::GetFunctions(TArray<FNiagaraFunctionSignat
     {
 		// GetNormal
 		FNiagaraFunctionSignature Sig;
-		Sig.Name = TEXT("GetNormal");
+		Sig.Name = GetNormalName;
 		Sig.bMemberFunction = true;
 		Sig.bRequiresContext = false;
 		Sig.Inputs.Add(FNiagaraVariable(FNiagaraTypeDefinition(GetClass()), TEXT("CSV")));			// CSV in
@@ -231,7 +309,7 @@ void UNiagaraDataInterfaceHoudiniCSV::GetFunctions(TArray<FNiagaraFunctionSignat
     {
 		// GetTime
 		FNiagaraFunctionSignature Sig;
-		Sig.Name = TEXT("GetTime");
+		Sig.Name = GetTimeName;
 		Sig.bMemberFunction = true;
 		Sig.bRequiresContext = false;
 		Sig.Inputs.Add(FNiagaraVariable(FNiagaraTypeDefinition(GetClass()), TEXT("CSV")));			// CSV in
@@ -247,7 +325,7 @@ void UNiagaraDataInterfaceHoudiniCSV::GetFunctions(TArray<FNiagaraFunctionSignat
 	{
 		// GetVelocity
 		FNiagaraFunctionSignature Sig;
-		Sig.Name = TEXT("GetVelocity");
+		Sig.Name = GetVelocityName;
 		Sig.bMemberFunction = true;
 		Sig.bRequiresContext = false;
 		Sig.Inputs.Add(FNiagaraVariable(FNiagaraTypeDefinition(GetClass()), TEXT("CSV")));			// CSV in
@@ -263,7 +341,7 @@ void UNiagaraDataInterfaceHoudiniCSV::GetFunctions(TArray<FNiagaraFunctionSignat
 	{
 		// GetColor
 		FNiagaraFunctionSignature Sig;
-		Sig.Name = TEXT("GetColor");
+		Sig.Name = GetColorName;
 		Sig.bMemberFunction = true;
 		Sig.bRequiresContext = false;
 		Sig.Inputs.Add(FNiagaraVariable(FNiagaraTypeDefinition(GetClass()), TEXT("CSV")));			// CSV in
@@ -279,7 +357,7 @@ void UNiagaraDataInterfaceHoudiniCSV::GetFunctions(TArray<FNiagaraFunctionSignat
     {
 		// GetNumberOfPoints
 		FNiagaraFunctionSignature Sig;
-		Sig.Name = TEXT("GetNumberOfPoints");
+		Sig.Name = GetNumberOfPointsName;
 		Sig.bMemberFunction = true;
 		Sig.bRequiresContext = false;
 		Sig.Inputs.Add(FNiagaraVariable(FNiagaraTypeDefinition(GetClass()), TEXT("CSV")));					// CSV in
@@ -294,7 +372,7 @@ void UNiagaraDataInterfaceHoudiniCSV::GetFunctions(TArray<FNiagaraFunctionSignat
 	{
 		// GetNumberOfRows
 		FNiagaraFunctionSignature Sig;
-		Sig.Name = TEXT("GetNumberOfRows");
+		Sig.Name = GetNumberOfRowsName;
 		Sig.bMemberFunction = true;
 		Sig.bRequiresContext = false;
 		Sig.Inputs.Add( FNiagaraVariable( FNiagaraTypeDefinition( GetClass()), TEXT("CSV") ) );					// CSV in
@@ -309,7 +387,7 @@ void UNiagaraDataInterfaceHoudiniCSV::GetFunctions(TArray<FNiagaraFunctionSignat
 	{
 		// GetNumberOfColumns
 		FNiagaraFunctionSignature Sig;
-		Sig.Name = TEXT("GetNumberOfColumns");
+		Sig.Name = GetNumberOfColumnsName;
 		Sig.bMemberFunction = true;
 		Sig.bRequiresContext = false;
 		Sig.Inputs.Add( FNiagaraVariable( FNiagaraTypeDefinition( GetClass()), TEXT("CSV") ) );						// CSV in
@@ -324,7 +402,7 @@ void UNiagaraDataInterfaceHoudiniCSV::GetFunctions(TArray<FNiagaraFunctionSignat
     {
 		// GetLastRowIndexAtTime
 		FNiagaraFunctionSignature Sig;
-		Sig.Name = TEXT("GetLastRowIndexAtTime");
+		Sig.Name = GetLastRowIndexAtTimeName;
 		Sig.bMemberFunction = true;
 		Sig.bRequiresContext = false;
 		Sig.Inputs.Add(FNiagaraVariable( FNiagaraTypeDefinition(GetClass()), TEXT("CSV") ) );					// CSV in
@@ -340,13 +418,13 @@ void UNiagaraDataInterfaceHoudiniCSV::GetFunctions(TArray<FNiagaraFunctionSignat
     {
 		// GetPointIdsToSpawnAtTime
 		FNiagaraFunctionSignature Sig;
-		Sig.Name = TEXT("GetPointIDsToSpawnAtTime");
+		Sig.Name = GetPointIDsToSpawnAtTimeName;
 		Sig.bMemberFunction = true;
 		Sig.bRequiresContext = false;
 		Sig.Inputs.Add(FNiagaraVariable( FNiagaraTypeDefinition(GetClass()), TEXT("CSV") ) );				// CSV in
 		Sig.Inputs.Add(FNiagaraVariable( FNiagaraTypeDefinition::GetFloatDef(), TEXT("Time") ) );		    // Time in
-		Sig.Outputs.Add(FNiagaraVariable( FNiagaraTypeDefinition::GetIntDef(), TEXT("MinIndex") ) );	    // Int Out
-		Sig.Outputs.Add(FNiagaraVariable( FNiagaraTypeDefinition::GetIntDef(), TEXT("MaxIndex") ) );	    // Int Out
+		Sig.Outputs.Add(FNiagaraVariable( FNiagaraTypeDefinition::GetIntDef(), TEXT("MinID") ) );			// Int Out
+		Sig.Outputs.Add(FNiagaraVariable( FNiagaraTypeDefinition::GetIntDef(), TEXT("MaxID") ) );			// Int Out
 		Sig.Outputs.Add(FNiagaraVariable( FNiagaraTypeDefinition::GetIntDef(), TEXT("Count") ) );		    // Int Out
 
 		Sig.SetDescription( LOCTEXT( "DataInterfaceHoudini_GetPointIDsToSpawnAtTime",
@@ -358,7 +436,7 @@ void UNiagaraDataInterfaceHoudiniCSV::GetFunctions(TArray<FNiagaraFunctionSignat
 	{
 		// GetRowIndexesForPointAtTime
 		FNiagaraFunctionSignature Sig;
-		Sig.Name = TEXT("GetRowIndexesForPointAtTime");
+		Sig.Name = GetRowIndexesForPointAtTimeName;
 		Sig.bMemberFunction = true;
 		Sig.bRequiresContext = false;
 		Sig.Inputs.Add(FNiagaraVariable(FNiagaraTypeDefinition(GetClass()), TEXT("CSV")));					// CSV in
@@ -377,11 +455,11 @@ void UNiagaraDataInterfaceHoudiniCSV::GetFunctions(TArray<FNiagaraFunctionSignat
 	{
 		// GetPointPositionAtTime
 		FNiagaraFunctionSignature Sig;
-		Sig.Name = TEXT("GetPointPositionAtTime");
+		Sig.Name = GetPointPositionAtTimeName;
 		Sig.bMemberFunction = true;
 		Sig.bRequiresContext = false;
 		Sig.Inputs.Add(FNiagaraVariable(FNiagaraTypeDefinition(GetClass()), TEXT("CSV")));				// CSV in
-		Sig.Inputs.Add(FNiagaraVariable(FNiagaraTypeDefinition::GetIntDef(), TEXT("PointID")));				// Point Number In
+		Sig.Inputs.Add(FNiagaraVariable(FNiagaraTypeDefinition::GetIntDef(), TEXT("PointID")));			// Point Number In
 		Sig.Inputs.Add(FNiagaraVariable(FNiagaraTypeDefinition::GetFloatDef(), TEXT("Time")));		    // Time in
 		Sig.Outputs.Add(FNiagaraVariable(FNiagaraTypeDefinition::GetVec3Def(), TEXT("Position")));		// Vector3 Out
 
@@ -394,11 +472,11 @@ void UNiagaraDataInterfaceHoudiniCSV::GetFunctions(TArray<FNiagaraFunctionSignat
 	{
 		// GetPointValueAtTime
 		FNiagaraFunctionSignature Sig;
-		Sig.Name = TEXT("GetPointValueAtTime");
+		Sig.Name = GetPointValueAtTimeName;
 		Sig.bMemberFunction = true;
 		Sig.bRequiresContext = false;
 		Sig.Inputs.Add(FNiagaraVariable(FNiagaraTypeDefinition(GetClass()), TEXT("CSV")));				// CSV in
-		Sig.Inputs.Add(FNiagaraVariable(FNiagaraTypeDefinition::GetIntDef(), TEXT("PointID")));				// Point Number In		
+		Sig.Inputs.Add(FNiagaraVariable(FNiagaraTypeDefinition::GetIntDef(), TEXT("PointID")));			// Point Number In		
 		Sig.Inputs.Add(FNiagaraVariable(FNiagaraTypeDefinition::GetIntDef(), TEXT("Col")));				// Col Index In
 		Sig.Inputs.Add(FNiagaraVariable(FNiagaraTypeDefinition::GetFloatDef(), TEXT("Time")));		    // Time in		
 		Sig.Outputs.Add(FNiagaraVariable(FNiagaraTypeDefinition::GetFloatDef(), TEXT("Value")));		// Float Out
@@ -412,7 +490,7 @@ void UNiagaraDataInterfaceHoudiniCSV::GetFunctions(TArray<FNiagaraFunctionSignat
 	{
 		// GetPointVectorValueAtTime
 		FNiagaraFunctionSignature Sig;
-		Sig.Name = TEXT("GetPointVectorValueAtTime");
+		Sig.Name = GetPointVectorValueAtTimeName;
 		Sig.bMemberFunction = true;
 		Sig.bRequiresContext = false;
 		Sig.Inputs.Add(FNiagaraVariable(FNiagaraTypeDefinition(GetClass()), TEXT("CSV")));				// CSV in
@@ -430,7 +508,7 @@ void UNiagaraDataInterfaceHoudiniCSV::GetFunctions(TArray<FNiagaraFunctionSignat
 	{
 		// GetPointVectorValueAtTimeEx
 		FNiagaraFunctionSignature Sig;
-		Sig.Name = TEXT("GetPointVectorValueAtTimeEx");
+		Sig.Name = GetPointVectorValueAtTimeExName;
 		Sig.bMemberFunction = true;
 		Sig.bRequiresContext = false;
 		Sig.Inputs.Add(FNiagaraVariable(FNiagaraTypeDefinition(GetClass()), TEXT("CSV")));				// CSV in
@@ -450,7 +528,7 @@ void UNiagaraDataInterfaceHoudiniCSV::GetFunctions(TArray<FNiagaraFunctionSignat
 	{
 		// GetPointLife
 		FNiagaraFunctionSignature Sig;
-		Sig.Name = TEXT("GetPointLife");
+		Sig.Name = GetPointLifeName;
 		Sig.bMemberFunction = true;
 		Sig.bRequiresContext = false;
 		Sig.Inputs.Add(FNiagaraVariable(FNiagaraTypeDefinition(GetClass()), TEXT("CSV")));				// CSV in
@@ -467,7 +545,7 @@ void UNiagaraDataInterfaceHoudiniCSV::GetFunctions(TArray<FNiagaraFunctionSignat
 	{
 		// GetPointLifeAtTime
 		FNiagaraFunctionSignature Sig;
-		Sig.Name = TEXT("GetPointLifeAtTime");
+		Sig.Name = GetPointLifeAtTimeName;
 		Sig.bMemberFunction = true;
 		Sig.bRequiresContext = false;
 		Sig.Inputs.Add(FNiagaraVariable(FNiagaraTypeDefinition(GetClass()), TEXT("CSV")));				// CSV in
@@ -485,7 +563,7 @@ void UNiagaraDataInterfaceHoudiniCSV::GetFunctions(TArray<FNiagaraFunctionSignat
 	{
 		// GetPointType
 		FNiagaraFunctionSignature Sig;
-		Sig.Name = TEXT("GetPointType");
+		Sig.Name = GetPointTypeName;
 		Sig.bMemberFunction = true;
 		Sig.bRequiresContext = false;
 		Sig.Inputs.Add(FNiagaraVariable(FNiagaraTypeDefinition(GetClass()), TEXT("CSV")));				// CSV in
@@ -499,89 +577,6 @@ void UNiagaraDataInterfaceHoudiniCSV::GetFunctions(TArray<FNiagaraFunctionSignat
 	}
 }
 
-// build the shader function HLSL; function name is passed in, as it's defined per-DI; that way, configuration could change
-// the HLSL in the spirit of a static switch
-// TODO: need a way to identify each specific function here
-// 
-bool UNiagaraDataInterfaceHoudiniCSV::GetFunctionHLSL(const FName& DefinitionFunctionName, FString InstanceFunctionName, FNiagaraDataInterfaceGPUParamInfo& ParamInfo, FString& OutHLSL)
-{
-    //FString BufferName = Descriptors[0].BufferParamName;
-    //FString SecondBufferName = Descriptors[1].BufferParamName;
-
-    /*
-    if (InstanceFunctionName.Contains( TEXT( "GetFloatValue") ) )
-    {
-	FString BufferName = Descriptors[0].BufferParamName;
-	OutHLSL += TEXT("void ") + InstanceFunctionName + TEXT("(in float In_Row, in float In_Col, out float Out_Value) \n{\n");
-	OutHLSL += TEXT("\t Out_Value = ") + BufferName + TEXT("[(int)(In_Row + ( In_Col * ") + FString::FromInt(NumberOfRows) + TEXT(") ) ];");
-	OutHLSL += TEXT("\n}\n");
-    }
-    else if (InstanceFunctionName.Contains(TEXT("GetPosition")))
-    {
-	FString BufferName = Descriptors[1].BufferParamName;
-	OutHLSL += TEXT("void ") + InstanceFunctionName + TEXT("(in float In_N, out float3 Out_Value) \n{\n");
-	OutHLSL += TEXT("\t Out_Value.x = ") + BufferName + TEXT("[(int)(In_N) ];");
-	OutHLSL += TEXT("\t Out_Value.y = ") + BufferName + TEXT("[(int)(In_N + ") + FString::FromInt(NumberOfRows) + TEXT(") ];");
-	OutHLSL += TEXT("\t Out_Value.z = ") + BufferName + TEXT("[(int)(In_N + ( 2 * ") + FString::FromInt(NumberOfRows) + TEXT(") ) ];");
-	OutHLSL += TEXT("\n}\n");
-    }
-    else if (InstanceFunctionName.Contains(TEXT("GetNormal")))
-    {
-	FString BufferName = Descriptors[2].BufferParamName;
-	OutHLSL += TEXT("void ") + InstanceFunctionName + TEXT("(in float In_N, out float3 Out_Value) \n{\n");
-	OutHLSL += TEXT("\t Out_Value.x = ") + BufferName + TEXT("[(int)(In_N) ];");
-	OutHLSL += TEXT("\t Out_Value.y = ") + BufferName + TEXT("[(int)(In_N + ") + FString::FromInt(NumberOfRows) + TEXT(") ];");
-	OutHLSL += TEXT("\t Out_Value.z = ") + BufferName + TEXT("[(int)(In_N + ( 2 * ") + FString::FromInt(NumberOfRows) + TEXT(") ) ];");
-	OutHLSL += TEXT("\n}\n");
-    }
-    else if (InstanceFunctionName.Contains(TEXT("GetTime")))
-    {
-	FString BufferName = Descriptors[3].BufferParamName;
-	OutHLSL += TEXT("void ") + InstanceFunctionName + TEXT("(in float In_N, out float Out_Value) \n{\n");
-	OutHLSL += TEXT("\t Out_Value = ") + BufferName + TEXT("[(int)(In_N) ];");
-	OutHLSL += TEXT("\n}\n");
-    }
-    else if (InstanceFunctionName.Contains( TEXT("GetNumberOfPointsInCSV") ) )
-    {
-	OutHLSL += TEXT("void ") + InstanceFunctionName + TEXT("( out int Out_Value ) \n{\n");
-	OutHLSL += TEXT("\t Out_Value = ") + FString::FromInt(NumberOfRows) + TEXT(";");
-	OutHLSL += TEXT("\n}\n");
-    }*/
-    /*else if (InstanceFunctionName.Contains(TEXT("GetPositionAndTime")))
-    {
-	OutHLSL += TEXT("void ") + FunctionName + TEXT("(in float In_N, out float4 Out_Value) \n{\n");
-	OutHLSL += TEXT("\t Out_Value.x = ") + BufferName + TEXT("[(int)(In_N) ];");
-	OutHLSL += TEXT("\t Out_Value.y = ") + BufferName + TEXT("[(int)(In_N + ") + FString::FromInt(NumberOfRows) + TEXT(") ];");
-	OutHLSL += TEXT("\t Out_Value.z = ") + BufferName + TEXT("[(int)(In_N + ( 2 * ") + FString::FromInt(NumberOfRows) + TEXT(") ) ];");
-	OutHLSL += TEXT("\t Out_Value.w = ") + SecondBufferName + TEXT("[(int)(In_N) ];");
-	OutHLSL += TEXT("\n}\n");
-    }*/
-
-    return !OutHLSL.IsEmpty();
-}
-
-// build buffer definition hlsl
-// 1. Choose a buffer name, add the data interface ID (important!)
-// 2. add a DIGPUBufferParamDescriptor to the array argument; that'll be passed on to the FNiagaraShader for binding to a shader param, that can
-// then later be found by name via FindDIBufferParam for setting; 
-// 3. store buffer declaration hlsl in OutHLSL
-// multiple buffers can be defined at once here
-//
-void UNiagaraDataInterfaceHoudiniCSV::GetParameterDefinitionHLSL(FNiagaraDataInterfaceGPUParamInfo& ParamInfo, FString& OutHLSL)
-{
-    FString BufferName = "CSVData_" + ParamInfo.DataInterfaceHLSLSymbol;
-    OutHLSL += TEXT("Buffer<float> ") + BufferName + TEXT(";\n");
-
-    BufferName = "PositionData_" + ParamInfo.DataInterfaceHLSLSymbol;
-    OutHLSL += TEXT("Buffer<float> ") + BufferName + TEXT(";\n");
-
-    BufferName = "NormalData_" + ParamInfo.DataInterfaceHLSLSymbol;
-    OutHLSL += TEXT("Buffer<float> ") + BufferName + TEXT(";\n");
-
-    BufferName = "TimeData_" + ParamInfo.DataInterfaceHLSLSymbol;
-    OutHLSL += TEXT("Buffer<float> ") + BufferName + TEXT(";\n");
-}
-
 DEFINE_NDI_RAW_FUNC_BINDER(UNiagaraDataInterfaceHoudiniCSV, GetFloatValue);
 //DEFINE_NDI_RAW_FUNC_BINDER(UNiagaraDataInterfaceHoudiniCSV, GetFloatValueByString);
 DEFINE_NDI_RAW_FUNC_BINDER(UNiagaraDataInterfaceHoudiniCSV, GetVectorValue);
@@ -604,97 +599,97 @@ DEFINE_NDI_RAW_FUNC_BINDER(UNiagaraDataInterfaceHoudiniCSV, GetPointLife);
 DEFINE_NDI_RAW_FUNC_BINDER(UNiagaraDataInterfaceHoudiniCSV, GetPointType);
 void UNiagaraDataInterfaceHoudiniCSV::GetVMExternalFunction(const FVMExternalFunctionBindingInfo& BindingInfo, void* InstanceData, FVMExternalFunction &OutFunc)
 {
-    if (BindingInfo.Name == TEXT("GetFloatValue") && BindingInfo.GetNumInputs() == 2 && BindingInfo.GetNumOutputs() == 1)
+    if (BindingInfo.Name == GetFloatValueName && BindingInfo.GetNumInputs() == 2 && BindingInfo.GetNumOutputs() == 1)
     {
 		TNDIParamBinder<0, int32, TNDIParamBinder<1, int32, NDI_RAW_FUNC_BINDER(UNiagaraDataInterfaceHoudiniCSV, GetFloatValue)>>::Bind(this, BindingInfo, InstanceData, OutFunc);
     }
-    /*else if (BindingInfo.Name == TEXT("GetFloatValueByString") && BindingInfo.GetNumInputs() == 2 && BindingInfo.GetNumOutputs() == 1)
+    /*else if (BindingInfo.Name == GetFloatValueByStringName && BindingInfo.GetNumInputs() == 2 && BindingInfo.GetNumOutputs() == 1)
     {
 		TNDIParamBinder<0, float, TNDIParamBinder<1, FString, NDI_RAW_FUNC_BINDER(UNiagaraDataInterfaceHoudiniCSV, GetFloatValueByString)>>::Bind(this, BindingInfo, InstanceData, OutFunc);
     }*/
-    else if (BindingInfo.Name == TEXT("GetVectorValue") && BindingInfo.GetNumInputs() == 2 && BindingInfo.GetNumOutputs() == 3)
+    else if (BindingInfo.Name == GetVectorValueName && BindingInfo.GetNumInputs() == 2 && BindingInfo.GetNumOutputs() == 3)
     {
 		TNDIParamBinder<0, int32, TNDIParamBinder<1, int32, NDI_RAW_FUNC_BINDER(UNiagaraDataInterfaceHoudiniCSV, GetVectorValue)>>::Bind(this, BindingInfo, InstanceData, OutFunc);
     }
-	else if (BindingInfo.Name == TEXT("GetVectorValueEx") && BindingInfo.GetNumInputs() == 4 && BindingInfo.GetNumOutputs() == 3)
+	else if (BindingInfo.Name == GetVectorValueExName && BindingInfo.GetNumInputs() == 4 && BindingInfo.GetNumOutputs() == 3)
 	{
 		TNDIParamBinder<0, int32, TNDIParamBinder<1, int32, TNDIParamBinder<2, bool, TNDIParamBinder<3, bool, NDI_RAW_FUNC_BINDER(UNiagaraDataInterfaceHoudiniCSV, GetVectorValueEx)>>>>::Bind(this, BindingInfo, InstanceData, OutFunc);
 	}
-    else if (BindingInfo.Name == TEXT("GetPosition") && BindingInfo.GetNumInputs() == 1 && BindingInfo.GetNumOutputs() == 3)
+    else if (BindingInfo.Name == GetPositionName && BindingInfo.GetNumInputs() == 1 && BindingInfo.GetNumOutputs() == 3)
     {
 		TNDIParamBinder<0, int32, NDI_RAW_FUNC_BINDER(UNiagaraDataInterfaceHoudiniCSV, GetPosition)>::Bind(this, BindingInfo, InstanceData, OutFunc);
     }
-    else if (BindingInfo.Name == TEXT("GetNormal") && BindingInfo.GetNumInputs() == 1 && BindingInfo.GetNumOutputs() == 3)
+    else if (BindingInfo.Name == GetNormalName && BindingInfo.GetNumInputs() == 1 && BindingInfo.GetNumOutputs() == 3)
     {
 		TNDIParamBinder<0, int32, NDI_RAW_FUNC_BINDER(UNiagaraDataInterfaceHoudiniCSV, GetNormal)>::Bind(this, BindingInfo, InstanceData, OutFunc);
     }
-    else if (BindingInfo.Name == TEXT("GetTime") && BindingInfo.GetNumInputs() == 1 && BindingInfo.GetNumOutputs() == 1)
+    else if (BindingInfo.Name == GetTimeName && BindingInfo.GetNumInputs() == 1 && BindingInfo.GetNumOutputs() == 1)
     {
 		TNDIParamBinder<0, int32, NDI_RAW_FUNC_BINDER(UNiagaraDataInterfaceHoudiniCSV, GetTime)>::Bind(this, BindingInfo, InstanceData, OutFunc);
     }
-	else if (BindingInfo.Name == TEXT("GetVelocity") && BindingInfo.GetNumInputs() == 1 && BindingInfo.GetNumOutputs() == 3)
+	else if (BindingInfo.Name == GetVelocityName && BindingInfo.GetNumInputs() == 1 && BindingInfo.GetNumOutputs() == 3)
 	{
 		TNDIParamBinder<0, int32, NDI_RAW_FUNC_BINDER(UNiagaraDataInterfaceHoudiniCSV, GetVelocity)>::Bind(this, BindingInfo, InstanceData, OutFunc);
 	}
-	else if (BindingInfo.Name == TEXT("GetColor") && BindingInfo.GetNumInputs() == 1 && BindingInfo.GetNumOutputs() == 4)
+	else if (BindingInfo.Name == GetColorName && BindingInfo.GetNumInputs() == 1 && BindingInfo.GetNumOutputs() == 4)
 	{
 		TNDIParamBinder<0, int32, NDI_RAW_FUNC_BINDER(UNiagaraDataInterfaceHoudiniCSV, GetColor)>::Bind(this, BindingInfo, InstanceData, OutFunc);
 	}
-    else if (BindingInfo.Name == TEXT("GetPositionAndTime") && BindingInfo.GetNumInputs() == 1 && BindingInfo.GetNumOutputs() == 4)
+    else if (BindingInfo.Name == GetPositionAndTimeName && BindingInfo.GetNumInputs() == 1 && BindingInfo.GetNumOutputs() == 4)
     {
 		TNDIParamBinder<0, int32, NDI_RAW_FUNC_BINDER(UNiagaraDataInterfaceHoudiniCSV, GetPositionAndTime)>::Bind(this, BindingInfo, InstanceData, OutFunc);
     }
-    else if ( BindingInfo.Name == TEXT("GetNumberOfPoints") && BindingInfo.GetNumInputs() == 0 && BindingInfo.GetNumOutputs() == 1 )
+    else if ( BindingInfo.Name == GetNumberOfPointsName && BindingInfo.GetNumInputs() == 0 && BindingInfo.GetNumOutputs() == 1 )
     {
 		OutFunc = FVMExternalFunction::CreateUObject(this, &UNiagaraDataInterfaceHoudiniCSV::GetNumberOfPoints);
     }
-	else if (BindingInfo.Name == TEXT("GetNumberOfRows") && BindingInfo.GetNumInputs() == 0 && BindingInfo.GetNumOutputs() == 1)
+	else if (BindingInfo.Name == GetNumberOfRowsName && BindingInfo.GetNumInputs() == 0 && BindingInfo.GetNumOutputs() == 1)
 	{
 		OutFunc = FVMExternalFunction::CreateUObject(this, &UNiagaraDataInterfaceHoudiniCSV::GetNumberOfRows);
 	}
-	else if (BindingInfo.Name == TEXT("GetNumberOfColumns") && BindingInfo.GetNumInputs() == 0 && BindingInfo.GetNumOutputs() == 1)
+	else if (BindingInfo.Name == GetNumberOfColumnsName && BindingInfo.GetNumInputs() == 0 && BindingInfo.GetNumOutputs() == 1)
 	{
 		OutFunc = FVMExternalFunction::CreateUObject(this, &UNiagaraDataInterfaceHoudiniCSV::GetNumberOfColumns);
 	}
-    else if (BindingInfo.Name == TEXT("GetLastRowIndexAtTime") && BindingInfo.GetNumInputs() == 1 && BindingInfo.GetNumOutputs() == 1)
+    else if (BindingInfo.Name == GetLastRowIndexAtTimeName && BindingInfo.GetNumInputs() == 1 && BindingInfo.GetNumOutputs() == 1)
     {
 		TNDIParamBinder<0, float, NDI_RAW_FUNC_BINDER(UNiagaraDataInterfaceHoudiniCSV, GetLastRowIndexAtTime)>::Bind(this, BindingInfo, InstanceData, OutFunc);
     }
-    else if (BindingInfo.Name == TEXT("GetPointIDsToSpawnAtTime") && BindingInfo.GetNumInputs() == 1 && BindingInfo.GetNumOutputs() == 3)
+    else if (BindingInfo.Name == GetPointIDsToSpawnAtTimeName && BindingInfo.GetNumInputs() == 1 && BindingInfo.GetNumOutputs() == 3)
     {
 		TNDIParamBinder<0, float, NDI_RAW_FUNC_BINDER(UNiagaraDataInterfaceHoudiniCSV, GetPointIDsToSpawnAtTime)>::Bind(this, BindingInfo, InstanceData, OutFunc);
     }
-	else if (BindingInfo.Name == TEXT("GetRowIndexesForPointAtTime") && BindingInfo.GetNumInputs() == 2 && BindingInfo.GetNumOutputs() == 3)
+	else if (BindingInfo.Name == GetRowIndexesForPointAtTimeName && BindingInfo.GetNumInputs() == 2 && BindingInfo.GetNumOutputs() == 3)
 	{
 		TNDIParamBinder<0, int32, TNDIParamBinder<1, float, NDI_RAW_FUNC_BINDER(UNiagaraDataInterfaceHoudiniCSV, GetRowIndexesForPointAtTime)>>::Bind(this, BindingInfo, InstanceData, OutFunc);
 	}
-	else if (BindingInfo.Name == TEXT("GetPointPositionAtTime") && BindingInfo.GetNumInputs() == 2 && BindingInfo.GetNumOutputs() == 3)
+	else if (BindingInfo.Name == GetPointPositionAtTimeName && BindingInfo.GetNumInputs() == 2 && BindingInfo.GetNumOutputs() == 3)
 	{
 		TNDIParamBinder<0, int32, TNDIParamBinder<1, float, NDI_RAW_FUNC_BINDER(UNiagaraDataInterfaceHoudiniCSV, GetPointPositionAtTime)>>::Bind(this, BindingInfo, InstanceData, OutFunc);
 	}
-	else if (BindingInfo.Name == TEXT("GetPointValueAtTime") && BindingInfo.GetNumInputs() == 3 && BindingInfo.GetNumOutputs() == 1)
+	else if (BindingInfo.Name == GetPointValueAtTimeName && BindingInfo.GetNumInputs() == 3 && BindingInfo.GetNumOutputs() == 1)
 	{
 		TNDIParamBinder<0, int32, TNDIParamBinder<1, int32, TNDIParamBinder<2, float, NDI_RAW_FUNC_BINDER(UNiagaraDataInterfaceHoudiniCSV, GetPointValueAtTime)>>>::Bind(this, BindingInfo, InstanceData, OutFunc);
 	}
-	else if (BindingInfo.Name == TEXT("GetPointVectorValueAtTime") && BindingInfo.GetNumInputs() == 3 && BindingInfo.GetNumOutputs() == 3)
+	else if (BindingInfo.Name == GetPointVectorValueAtTimeName && BindingInfo.GetNumInputs() == 3 && BindingInfo.GetNumOutputs() == 3)
 	{
 		TNDIParamBinder<0, int32, TNDIParamBinder<1, int32, TNDIParamBinder<2, float, NDI_RAW_FUNC_BINDER(UNiagaraDataInterfaceHoudiniCSV, GetPointVectorValueAtTime)>>>::Bind(this, BindingInfo, InstanceData, OutFunc);
 	}
-	else if (BindingInfo.Name == TEXT("GetPointVectorValueAtTimeEx") && BindingInfo.GetNumInputs() == 5 && BindingInfo.GetNumOutputs() == 3)
+	else if (BindingInfo.Name == GetPointVectorValueAtTimeExName && BindingInfo.GetNumInputs() == 5 && BindingInfo.GetNumOutputs() == 3)
 	{
 		TNDIParamBinder<0, int32, TNDIParamBinder<1, int32, TNDIParamBinder<2, float, TNDIParamBinder<3, bool, TNDIParamBinder<4, bool, NDI_RAW_FUNC_BINDER(UNiagaraDataInterfaceHoudiniCSV, GetPointVectorValueAtTimeEx)>>>>>::Bind(this, BindingInfo, InstanceData, OutFunc);
 	}
-	else if (BindingInfo.Name == TEXT("GetPointLife") && BindingInfo.GetNumInputs() == 1 && BindingInfo.GetNumOutputs() == 1)
+	else if (BindingInfo.Name == GetPointLifeName && BindingInfo.GetNumInputs() == 1 && BindingInfo.GetNumOutputs() == 1)
 	{
 		TNDIParamBinder<0, int32, NDI_RAW_FUNC_BINDER(UNiagaraDataInterfaceHoudiniCSV, GetPointLife)>::Bind(this, BindingInfo, InstanceData, OutFunc);
 	}
 	/*
-	else if (BindingInfo.Name == TEXT("GetPointLifeAtTime") && BindingInfo.GetNumInputs() == 2 && BindingInfo.GetNumOutputs() == 1)
+	else if (BindingInfo.Name == GetPointLifeAtTimeName && BindingInfo.GetNumInputs() == 2 && BindingInfo.GetNumOutputs() == 1)
 	{
 		TNDIParamBinder<0, int32, TNDIParamBinder<1, float, NDI_RAW_FUNC_BINDER(UNiagaraDataInterfaceHoudiniCSV, GetPointLifeAtTime)>>::Bind(this, BindingInfo, InstanceData, OutFunc);
 	}
 	*/
-	else if (BindingInfo.Name == TEXT("GetPointType") && BindingInfo.GetNumInputs() == 1 && BindingInfo.GetNumOutputs() == 1)
+	else if (BindingInfo.Name == GetPointTypeName && BindingInfo.GetNumInputs() == 1 && BindingInfo.GetNumOutputs() == 1)
 	{
 		TNDIParamBinder<0, int32, NDI_RAW_FUNC_BINDER(UNiagaraDataInterfaceHoudiniCSV, GetPointType)>::Bind(this, BindingInfo, InstanceData, OutFunc);
 	}
@@ -988,51 +983,10 @@ void UNiagaraDataInterfaceHoudiniCSV::GetPointIDsToSpawnAtTime( FVectorVMContext
 
 		int32 value = 0;
 		int32 min = 0, max = 0, count = 0;
+
 		if ( HoudiniCSVAsset )
 		{
-			if ( !HoudiniCSVAsset->GetLastPointIDToSpawnAtTime( t, value ) )
-			{
-				// The CSV file doesn't have time informations, so always return all points in the file
-				min = 0;
-				max = value;
-				count = max - min + 1;
-			}
-			else
-			{
-				// The CSV file has time informations
-				// First, detect if we need to reset LastSpawnedPointID (after a loop of the emitter)
-				if ( value < LastSpawnedPointID )
-					LastSpawnedPointID = -1;
-
-				if ( value < 0 )
-				{
-					// Nothing to spawn, t is lower than the point's time
-					LastSpawnedPointID = -1;
-				}
-				else
-				{
-					// The last time value in the CSV is lower than t, spawn everything if we didnt already!
-					if ( value >= HoudiniCSVAsset->GetNumberOfPoints() )
-						value = value - 1;
-
-					if ( value == LastSpawnedPointID)
-					{
-						// We dont have any new point to spawn
-						min = value;
-						max = value;
-						count = 0;
-					}
-					else
-					{
-						// We have points to spawn at time t
-						min = LastSpawnedPointID + 1;
-						max = value;
-						count = max - min + 1;
-
-						LastSpawnedPointID = max;
-					}
-				}
-			}
+			HoudiniCSVAsset->GetPointIDsToSpawnAtTime(t, min, max, count, LastSpawnedPointID, LastSpawnTime);
 		}
 
 		*OutMinValue.GetDest() = min;
@@ -1344,8 +1298,8 @@ void UNiagaraDataInterfaceHoudiniCSV::GetPointType(FVectorVMContext& Context)
 
 void UNiagaraDataInterfaceHoudiniCSV::GetNumberOfRows(FVectorVMContext& Context)
 {
-    FRegisterHandler<int32> OutNumRows(Context);
-    *OutNumRows.GetDest() = HoudiniCSVAsset ? HoudiniCSVAsset->GetNumberOfRows() : 0;
+	FRegisterHandler<int32> OutNumRows(Context);
+	*OutNumRows.GetDest() = HoudiniCSVAsset ? HoudiniCSVAsset->GetNumberOfRows() : 0;
 	OutNumRows.Advance();
 }
 
@@ -1363,4 +1317,778 @@ void UNiagaraDataInterfaceHoudiniCSV::GetNumberOfPoints(FVectorVMContext& Contex
 	OutNumPoints.Advance();
 }
 
+// Build the shader function HLSL Code.
+bool UNiagaraDataInterfaceHoudiniCSV::GetFunctionHLSL(const FName& DefinitionFunctionName, FString InstanceFunctionName, FNiagaraDataInterfaceGPUParamInfo& ParamInfo, FString& OutHLSL)
+{
+	// Build the buffer/variable names declared for this DI
+	FString NumberOfRowsVar = NumberOfRowsBaseName + ParamInfo.DataInterfaceHLSLSymbol;
+	FString NumberOfColumnsVar = NumberOfColumnsBaseName + ParamInfo.DataInterfaceHLSLSymbol;
+	FString NumberOfPointsVar = NumberOfPointsBaseName + ParamInfo.DataInterfaceHLSLSymbol;
+	FString FloatBufferVar = FloatValuesBufferBaseName + ParamInfo.DataInterfaceHLSLSymbol;
+	FString AttributeColumnIndexesBuffer = SpecialAttributesColumnIndexesBufferBaseName + ParamInfo.DataInterfaceHLSLSymbol;
+	FString SpawnTimeBuffer = SpawnTimesBufferBaseName + ParamInfo.DataInterfaceHLSLSymbol;
+	FString LifeValuesBuffer = LifeValuesBufferBaseName + ParamInfo.DataInterfaceHLSLSymbol;
+	FString PointTypesBuffer = PointTypesBufferBaseName + ParamInfo.DataInterfaceHLSLSymbol;
+	FString MaxNumberOfIndexesPerPointVar = MaxNumberOfIndexesPerPointBaseName + ParamInfo.DataInterfaceHLSLSymbol;
+	FString PointValueIndexesBuffer = PointValueIndexesBufferBaseName + ParamInfo.DataInterfaceHLSLSymbol;
+	FString LastSpawnedPointIDVar = LastSpawnedPointIdBaseName + ParamInfo.DataInterfaceHLSLSymbol;
+	FString LastSpawnTimeVar = LastSpawnTimeBaseName + ParamInfo.DataInterfaceHLSLSymbol;
+
+	// Lambda returning the HLSL code used for reading a Float value in the FloatBuffer
+	auto ReadFloatInBuffer = [&](const FString& OutFloatValue, const FString& FloatRowIndex, const FString& FloatColIndex)
+	{
+		// \t OutValue = FloatBufferName[ (RowIndex) + ( (ColIndex) * (NumberOfRowsName) ) ];\n
+		return TEXT("\t ") + OutFloatValue + TEXT(" = ") + FloatBufferVar + TEXT("[ (") + FloatRowIndex + TEXT(") + ( (") + FloatColIndex + TEXT(") * (") + NumberOfRowsVar + TEXT(") ) ];\n");
+	};
+
+	// Lambda returning the HLSL code for reading a Vector value in the FloatBuffer
+	// It expects the In_DoSwap and In_DoScale integers to be defined before being called!
+	auto ReadVectorInBuffer = [&](const FString& OutVectorValue, const FString& VectorRowIndex, const FString& VectorColIndex)
+	{
+		FString OutHLSLCode;
+		OutHLSLCode += TEXT("\t// ReadVectorInBuffer\n");
+		OutHLSLCode += TEXT("\t{\n");
+			OutHLSLCode += TEXT("\t\tfloat3 temp_Value = float3(0.0, 0.0, 0.0);\n");
+			OutHLSLCode += ReadFloatInBuffer(TEXT("temp_Value.x"), VectorRowIndex, VectorColIndex);
+			OutHLSLCode += ReadFloatInBuffer(TEXT("temp_Value.y"), VectorRowIndex, VectorColIndex + TEXT(" + 1"));
+			OutHLSLCode += ReadFloatInBuffer(TEXT("temp_Value.z"), VectorRowIndex, VectorColIndex + TEXT(" + 2"));
+			OutHLSLCode += TEXT("\t\t") + OutVectorValue + TEXT(" = temp_Value;\n");
+			OutHLSLCode += TEXT("\t\tif ( In_DoSwap == 1 )\n");
+			OutHLSLCode += TEXT("\t\t{\n");
+				OutHLSLCode += TEXT("\t\t\t") + OutVectorValue + TEXT(".y= temp_Value.z;\n");
+				OutHLSLCode += TEXT("\t\t\t") + OutVectorValue + TEXT(".z= temp_Value.y;\n");
+			OutHLSLCode += TEXT("\t\t}\n");
+			OutHLSLCode += TEXT("\t\tif ( In_DoScale == 1 )\n");
+			OutHLSLCode += TEXT("\t\t{\n");
+				OutHLSLCode += TEXT("\t\t\t") + OutVectorValue + TEXT(" *= 100.0f;\n");
+			OutHLSLCode += TEXT("\t\t}\n");
+		OutHLSLCode += TEXT("\t}\n");
+		return OutHLSLCode;
+	};
+
+	// Lambda returning the HLSL code for reading a special attribute's column index in the buffer
+	auto GetSpecAttributeColumnIndex = [&](const EHoudiniAttributes& Attr)
+	{
+		FString OutHLSLCode;
+		OutHLSLCode += AttributeColumnIndexesBuffer + TEXT("[") + FString::FromInt(Attr) + TEXT("]");
+		return OutHLSLCode;
+	};
+
+	// Lambda returning the HLSL code for getting the previous and next row indexes for a given particle at a given time	
+	auto GetRowIndexesForPointAtTime = [&](const FString& In_PointID, const FString& In_Time, const FString& Out_PreviousRow, const FString& Out_NextRow, const FString& Out_Weight)
+	{
+		FString OutHLSLCode;
+		OutHLSLCode += TEXT("\t// GetRowIndexesForPointAtTime\n");
+		OutHLSLCode += TEXT("\t{\n");
+
+			OutHLSLCode += TEXT("\t\tfloat prev_time = -1.0f;\n");
+			OutHLSLCode += TEXT("\t\tfloat next_time = -1.0f;\n");
+
+			// Look at all the values for this Point
+			OutHLSLCode += TEXT("\t\tfor( int n = 0; n < ") + MaxNumberOfIndexesPerPointVar + TEXT("; n++ )\n\t{\n");
+				OutHLSLCode += TEXT("\t\t\tint current_row = ") + PointValueIndexesBuffer + TEXT("[ (") + In_PointID + TEXT(") * ") + MaxNumberOfIndexesPerPointVar + TEXT(" + n ];\n");
+				OutHLSLCode += TEXT("\t\t\tif ( current_row < 0 ){ break; }\n");
+
+				OutHLSLCode += TEXT("\t\t\tfloat current_time = -1.0f;\n");
+				OutHLSLCode += TEXT("\t\t\tint time_col = ") + GetSpecAttributeColumnIndex(EHoudiniAttributes::TIME) + TEXT(";\n");
+				OutHLSLCode += TEXT("\t\t\tif ( time_col < 0 || time_col >= ") + NumberOfColumnsVar + TEXT(" )\n");
+					OutHLSLCode += TEXT("\t\t\t\t{ current_time = 0.0f; }\n");
+				OutHLSLCode += TEXT("\t\t\telse\n");
+					OutHLSLCode += TEXT("\t\t\t{") + ReadFloatInBuffer(TEXT("current_time"), TEXT("current_row"), TEXT("time_col")) + TEXT(" }\n");
+
+				OutHLSLCode += TEXT("\t\t\tif ( current_time == (") + In_Time + TEXT(") )\n");
+					OutHLSLCode += TEXT("\t\t\t\t{ ") + Out_PreviousRow + TEXT(" = current_row; ") + Out_NextRow + TEXT(" = current_row; ") + Out_Weight + TEXT(" = 1.0; return;}\n");
+				OutHLSLCode += TEXT("\t\t\telse if ( current_time < (") + In_Time + TEXT(") ){\n");
+					OutHLSLCode += TEXT("\t\t\t\tif ( prev_time == -1.0f || prev_time < current_time )\n");
+						OutHLSLCode += TEXT("\t\t\t\t\t { ") + Out_PreviousRow + TEXT(" = current_row; prev_time = current_time; }\n");
+				OutHLSLCode += TEXT("\t\t\t}\n");
+				OutHLSLCode += TEXT("\t\t\telse if ( next_time == -1.0f || next_time > current_time )\n");
+					OutHLSLCode += TEXT("\t\t\t\t { ") + Out_NextRow + TEXT(" = current_row; next_time = current_time; break; }\n");
+			OutHLSLCode += TEXT("\t\t}\n");
+
+			OutHLSLCode += TEXT("\t\tif ( ") + Out_PreviousRow + TEXT(" < 0 )\n");
+				OutHLSLCode += TEXT("\t\t\t{ ") + Out_Weight + TEXT(" = 0.0f; ") + Out_PreviousRow + TEXT(" = ") + Out_NextRow + TEXT("; return;}\n");
+			OutHLSLCode += TEXT("\t\tif ( ") + Out_NextRow + TEXT(" < 0 )\n");
+				OutHLSLCode += TEXT("\t\t\t{ ") + Out_Weight + TEXT(" = 1.0f; ") + Out_NextRow + TEXT(" = ") + Out_PreviousRow + TEXT("; return;}\n");
+
+				// Calculate the weight
+			OutHLSLCode += TEXT("\t\t") + Out_Weight + TEXT(" = ( ( (") + In_Time + TEXT(") - prev_time ) / ( next_time - prev_time ) );\n");
+
+		OutHLSLCode += TEXT("\t}\n");
+		return OutHLSLCode;
+	};
+
+	// Build each function's HLSL code
+	if (DefinitionFunctionName == GetFloatValueName)
+	{
+		// GetFloatValue(int In_Row, int In_Col, out float Out_Value)		
+		OutHLSL += TEXT("void ") + InstanceFunctionName + TEXT("(int In_Row, int In_Col, out float Out_Value) \n{\n");
+
+			OutHLSL += ReadFloatInBuffer(TEXT("Out_Value"), TEXT("In_Row"), TEXT("In_Col"));
+
+		OutHLSL += TEXT("\n}\n");
+		return true;
+	}
+	/*
+	else if (DefinitionFunctionName == GetFloatValueByStringName)
+	{
+		// GetFloatValueByString(int In_Row, string In_ColTitle, out float Out_Value)
+		FString BufferName = FloatValuesBufferBaseName + ParamInfo.DataInterfaceHLSLSymbol;
+		OutHLSL += TEXT("void ") + InstanceFunctionName + TEXT("(int In_Row, string In_ColTitle, out float Out_Value) \n{\n");
+		// NEED CODE
+		OutHLSL += TEXT("\n}\n");
+		return true;
+	}
+	*/
+	else if (DefinitionFunctionName == GetVectorValueName)
+	{
+		// GetVectorValue(int In_Row, int In_Col, out float3 Out_Value)
+		OutHLSL += TEXT("void ") + InstanceFunctionName + TEXT("(int In_Row, int In_Col, out float3 Out_Value) \n{\n");
+
+			OutHLSL += TEXT("\tint In_DoSwap = 1;\n");
+			OutHLSL += TEXT("\tint In_DoScale = 1;\n");
+			OutHLSL += ReadVectorInBuffer(TEXT("Out_Value"), TEXT("In_Row"), TEXT("In_Col"));
+
+		OutHLSL += TEXT("\n}\n");
+		return true;
+	}
+	else if (DefinitionFunctionName == GetVectorValueExName)
+	{
+		// GetVectorValueEx(int In_Row, int In_Col, int In_DoSwap, int In_DoScale, out float3 Out_Value)
+		OutHLSL += TEXT("void ") + InstanceFunctionName + TEXT("(int In_Row, int In_Col, int In_DoSwap, int In_DoScale, out float3 Out_Value) \n{\n");
+		
+			OutHLSL += ReadVectorInBuffer(TEXT("Out_Value"), TEXT("In_Row"), TEXT("In_Col"));
+
+		OutHLSL += TEXT("\n}\n");
+		return true;
+	}
+	else if (DefinitionFunctionName == GetPositionName)
+	{
+		// GetPosition(int In_Row, out float3 Out_Position)
+		OutHLSL += TEXT("void ") + InstanceFunctionName + TEXT("(int In_Row, out float3 Out_Position) \n{\n");
+
+			OutHLSL += TEXT("\tint In_DoSwap = 1;\n");
+			OutHLSL += TEXT("\tint In_DoScale = 1;\n");
+			OutHLSL += TEXT("\tint In_Col = ") + GetSpecAttributeColumnIndex( EHoudiniAttributes::POSITION ) + TEXT(";\n");
+			OutHLSL += ReadVectorInBuffer(TEXT("Out_Position"), TEXT("In_Row"), TEXT("In_Col"));
+
+		OutHLSL += TEXT("\n}\n");
+		return true;
+	}
+	else if (DefinitionFunctionName == GetPositionAndTimeName)
+	{
+		// GetPositionAndTime(int In_Row, out float3 Out_Position, out float Out_Time)
+		OutHLSL += TEXT("void ") + InstanceFunctionName + TEXT("(int In_Row, out float3 Out_Position, out float Out_Time) \n{\n");
+
+			OutHLSL += TEXT("\tint In_DoSwap = 1;\n");
+			OutHLSL += TEXT("\tint In_DoScale = 1;\n");
+			OutHLSL += TEXT("\tint In_PosCol = ") + GetSpecAttributeColumnIndex(EHoudiniAttributes::POSITION) + TEXT(";\n");
+			OutHLSL += ReadVectorInBuffer(TEXT("Out_Position"), TEXT("In_Row"), TEXT("In_PosCol"));
+
+			OutHLSL += TEXT("\tint In_TimeCol = ") + GetSpecAttributeColumnIndex(EHoudiniAttributes::TIME) + TEXT(";\n");
+			OutHLSL += ReadFloatInBuffer(TEXT("Out_Time"), TEXT("In_Row"), TEXT("In_TimeCol"));
+
+		OutHLSL += TEXT("\n}\n");
+		return true;
+	}
+	else if (DefinitionFunctionName == GetNormalName)
+	{
+		// GetNormal(int In_Row, out float3 Out_Value)
+		OutHLSL += TEXT("void ") + InstanceFunctionName + TEXT("(int In_Row, out float3 Out_Normal) \n{\n");
+
+			OutHLSL += TEXT("\tint In_DoSwap = 1;\n");
+			OutHLSL += TEXT("\tint In_DoScale = 0;\n");
+			OutHLSL += TEXT("\tint In_Col = ") + GetSpecAttributeColumnIndex(EHoudiniAttributes::NORMAL) + TEXT(";\n");
+			OutHLSL += ReadVectorInBuffer(TEXT("Out_Normal"), TEXT("In_Row"), TEXT("In_Col"));
+
+		OutHLSL += TEXT("\n}\n");
+		return true;
+	}
+	else if (DefinitionFunctionName == GetTimeName)
+	{
+		// GetTime(int In_Row, out float Out_Value)
+		OutHLSL += TEXT("void ") + InstanceFunctionName + TEXT("(int In_Row, out float Out_Time) \n{\n");
+
+			OutHLSL += TEXT("\tint In_TimeCol = ") + GetSpecAttributeColumnIndex(EHoudiniAttributes::TIME) + TEXT(";\n");
+			OutHLSL += ReadFloatInBuffer(TEXT("Out_Time"), TEXT("In_Row"), TEXT("In_TimeCol"));
+
+		OutHLSL += TEXT("\n}\n");
+		return true;
+	}
+	else if (DefinitionFunctionName == GetVelocityName)
+	{
+		// GetVelocity(int In_Row, out float3 Out_Value)
+		OutHLSL += TEXT("void ") + InstanceFunctionName + TEXT("(int In_Row, out float3 Out_Velocity) \n{\n");
+
+			OutHLSL += TEXT("\tint In_DoSwap = 1;\n");
+			OutHLSL += TEXT("\tint In_DoScale = 0;\n");
+			OutHLSL += TEXT("\tint In_Col = ") + GetSpecAttributeColumnIndex(EHoudiniAttributes::VELOCITY) + TEXT(";\n");
+			OutHLSL += ReadVectorInBuffer(TEXT("Out_Velocity"), TEXT("In_Row"), TEXT("In_Col"));
+
+		OutHLSL += TEXT("\n}\n");
+		return true;
+	}
+	else if (DefinitionFunctionName == GetColorName)
+	{
+		// GetColor(int In_Row, out float4 Out_Value)
+		OutHLSL += TEXT("void ") + InstanceFunctionName + TEXT("(int In_Row, out float4 Out_Color) \n{\n");
+
+			OutHLSL += TEXT("\tfloat3 temp_color = float3(1.0, 1.0, 1.0);\n");
+			OutHLSL += TEXT("\tint In_ColorCol = ") + GetSpecAttributeColumnIndex(EHoudiniAttributes::COLOR) + TEXT(";\n");
+			OutHLSL += TEXT("\tint In_DoSwap = 0;\n");
+			OutHLSL += TEXT("\tint In_DoScale = 0;\n");
+			OutHLSL += ReadVectorInBuffer(TEXT("temp_color"), TEXT("In_Row"), TEXT("In_ColorCol"));
+
+			OutHLSL += TEXT("\tfloat temp_alpha = 1.0;\n");
+			OutHLSL += TEXT("\tint In_TimeCol = ") + GetSpecAttributeColumnIndex(EHoudiniAttributes::TIME) + TEXT(";\n");
+			OutHLSL += ReadFloatInBuffer(TEXT("temp_alpha"), TEXT("In_Row"), TEXT("In_TimeCol"));
+
+			OutHLSL += TEXT("Out_Color.x = temp_color.x;\n");
+			OutHLSL += TEXT("Out_Color.y = temp_color.y;\n");
+			OutHLSL += TEXT("Out_Color.z = temp_color.z;\n");
+			OutHLSL += TEXT("Out_Color.w = temp_alpha;\n");
+
+		OutHLSL += TEXT("\n}\n");
+		return true;
+	}
+	else if (DefinitionFunctionName == GetNumberOfPointsName)
+	{
+		// GetNumberOfPoints(out int Out_Value)
+		OutHLSL += TEXT("void ") + InstanceFunctionName + TEXT("(out int Out_NumPoints) \n{\n");
+			OutHLSL += TEXT("\tOut_NumPoints = ") + NumberOfPointsVar + TEXT(";\n");
+		OutHLSL += TEXT("\n}\n");
+		return true;
+	}
+	else if (DefinitionFunctionName == GetNumberOfRowsName)
+	{
+		// GetNumberOfRows(out int Out_Value)
+		OutHLSL += TEXT("void ") + InstanceFunctionName + TEXT("(out int Out_NumRows) \n{\n");
+			OutHLSL += TEXT("\tOut_NumRows = ") + NumberOfRowsVar + TEXT(";\n");
+		OutHLSL += TEXT("\n}\n");
+		return true;
+	}
+	else if (DefinitionFunctionName == GetNumberOfColumnsName)
+	{
+		// GetNumberOfColumns(out int Out_Value)
+		OutHLSL += TEXT("void ") + InstanceFunctionName + TEXT("(out int Out_NumColumns) \n{\n");
+			OutHLSL += TEXT("\tOut_NumColumns = ") + NumberOfColumnsVar + TEXT(";\n");
+		OutHLSL += TEXT("\n}\n");
+		return true;
+	}
+	else if (DefinitionFunctionName == GetLastRowIndexAtTimeName)
+	{
+		// GetLastRowIndexAtTime(float In_Time, out int Out_Value)
+		OutHLSL += TEXT("void ") + InstanceFunctionName + TEXT("(float In_Time, out int Out_Value) \n{\n");
+
+			OutHLSL += TEXT("\tint In_TimeCol = ") + GetSpecAttributeColumnIndex( EHoudiniAttributes::TIME ) + TEXT(";\n");
+			OutHLSL += TEXT("\tfloat temp_time = 1.0;\n");
+			OutHLSL += ReadFloatInBuffer(TEXT("temp_time"), NumberOfRowsVar + TEXT("- 1"), TEXT("In_TimeCol"));
+			OutHLSL += TEXT("\tif ( temp_time < In_Time ) { Out_Value = ") + NumberOfRowsVar + TEXT(" - 1; return; }\n");
+			OutHLSL += TEXT("\tint lastRow = -1;\n");
+			OutHLSL += TEXT("\tfor( int n = 0; n < ") + NumberOfRowsVar + TEXT("; n++ )\n\t{\n");
+				OutHLSL += TEXT("\t") + ReadFloatInBuffer(TEXT("temp_time"), TEXT("n"), TEXT("In_TimeCol"));
+				OutHLSL += TEXT("\t\tif ( temp_time == In_Time ){ lastRow = n ;}");
+				OutHLSL += TEXT("\t\telse if ( temp_time > In_Time ){ lastRow = n -1; return;}");
+				OutHLSL += TEXT("\t\tif ( lastRow == -1 ){ lastRow = ") + NumberOfRowsVar + TEXT(" - 1; }");
+			OutHLSL += TEXT("\t}\n");
+
+		OutHLSL += TEXT("\n}\n");
+		return true;
+	}
+	else if (DefinitionFunctionName == GetPointIDsToSpawnAtTimeName)
+	{
+		// GetPointIDsToSpawnAtTime(float In_Time, out int Out_MinID, out int Out_MaxID, out int Out_Count)
+		OutHLSL += TEXT("void ") + InstanceFunctionName + TEXT("(float In_Time, out int Out_MinID, out int Out_MaxID, out int Out_Count) \n{\n");
+			
+			OutHLSL += TEXT("\tint time_col = ") + GetSpecAttributeColumnIndex(EHoudiniAttributes::TIME) + TEXT(";\n");
+			OutHLSL += TEXT("\tif( time_col < 0 )\n");
+				OutHLSL += TEXT("\t\t{ Out_Count = ") + NumberOfPointsVar +TEXT("; Out_MinID = 0; Out_MaxID = Out_Count - 1; return; }\n");
+
+			// GetLastPointIDToSpawnAtTime
+			OutHLSL += TEXT("\tint last_id = -1;\n");
+			OutHLSL += TEXT("\tfloat temp_time = 0;\n");
+			OutHLSL += TEXT("\tfor( int n = 0; n < ") + NumberOfPointsVar + TEXT("; n++ )\n\t{\n");
+				OutHLSL += TEXT("\t\ttemp_time = ") + SpawnTimeBuffer + TEXT("[ n ];\n");
+				OutHLSL += TEXT("\t\tif ( temp_time == In_Time )\n");
+					OutHLSL += TEXT("\t\t\t{ last_id = n; }\n");
+				OutHLSL += TEXT("\t\telse if ( temp_time > In_Time )\n");
+					OutHLSL += TEXT("\t\t\t{ last_id = n - 1; break; }\n");
+			OutHLSL += TEXT("\t}\n");
+
+			// Check if we need to reset lastspawnedPointID
+			OutHLSL += TEXT("\tif ( last_id < ") + LastSpawnedPointIDVar + TEXT(" || In_Time <= ") + LastSpawnTimeVar + TEXT(")\n");
+				OutHLSL += TEXT("\t\t{") + LastSpawnedPointIDVar + TEXT(" = -1; }\n");
+			// Nothing to spawn, In_Time is lower than the point's time
+			OutHLSL += TEXT("\tif ( last_id < 0 )\n");
+				OutHLSL += TEXT("\t\t{") + LastSpawnedPointIDVar + TEXT(" = -1;Out_Count = 0; Out_MinID = 0; Out_MaxID = 0;return;}\n");
+			// We dont have any new point to spawn
+			OutHLSL += TEXT("\tif ( last_id == ") + LastSpawnedPointIDVar + TEXT(")\n");
+				OutHLSL += TEXT("{ Out_MinID = last_id; Out_MaxID = last_id; Out_Count = 0; return;}\n");
+
+			OutHLSL += TEXT("\tOut_MinID = ") + LastSpawnedPointIDVar + TEXT(" + 1;\n");
+			OutHLSL += TEXT("\tOut_MaxID = last_id;\n");
+			OutHLSL += TEXT("\tOut_Count = Out_MaxID - Out_MinID + 1;\n");
+			OutHLSL += TEXT("\t") + LastSpawnedPointIDVar + TEXT(" = Out_MaxID;\n");
+			OutHLSL += TEXT("\t") + LastSpawnTimeVar + TEXT(" = In_Time;\n");
+
+		OutHLSL += TEXT("\n}\n");
+		return true;
+	}
+	else if (DefinitionFunctionName == GetRowIndexesForPointAtTimeName)
+	{
+		// GetRowIndexesForPointAtTime(int In_PointID, float In_Time, out int Out_PreviousRow, out int Out_NextRow, out float Out_Weight)
+		OutHLSL += TEXT("void ") + InstanceFunctionName + TEXT("(int In_PointID, float In_Time, out int Out_PreviousRow, out int Out_NextRow, out float Out_Weight) \n{\n");
+
+			OutHLSL += TEXT("\tOut_PreviousRow = 0;\n\tOut_NextRow = 0;\n\tOut_Weight = 0;\n");
+			OutHLSL += GetRowIndexesForPointAtTime(TEXT("In_PointID"), TEXT("In_Time"), TEXT("Out_PreviousRow"), TEXT("Out_NextRow"), TEXT("Out_Weight"));
+
+		OutHLSL += TEXT("\n}\n");
+		return true;
+	}
+	else if (DefinitionFunctionName == GetPointPositionAtTimeName)
+	{
+		// GetPointPositionAtTime(int In_PointID, float In_Time, out float3 Out_Value)
+		OutHLSL += TEXT("void ") + InstanceFunctionName + TEXT("(int In_PointID, float In_Time, out float3 Out_Value) \n{\n");
+		OutHLSL += TEXT("Out_Value = float3( 0.0, 0.0, 0.0 );\n");
+			OutHLSL += TEXT("\tint pos_col = ") + GetSpecAttributeColumnIndex(EHoudiniAttributes::POSITION) + TEXT(";\n");
+			OutHLSL += TEXT("\tint prev_index = 0;int next_index = 0;float weight = 0.0f;\n");
+			OutHLSL += GetRowIndexesForPointAtTime(TEXT("In_PointID"), TEXT("In_Time"), TEXT("prev_index"), TEXT("next_index"), TEXT("weight"));
+
+			OutHLSL += TEXT("\tint In_DoSwap = 1;\n");
+			OutHLSL += TEXT("\tint In_DoScale = 1;\n");
+
+			OutHLSL += TEXT("\tfloat3 prev_vector = float3( 0.0, 0.0, 0.0 );\n");
+			OutHLSL += ReadVectorInBuffer(TEXT("prev_vector"), TEXT("prev_index"), TEXT("pos_col"));
+			OutHLSL += TEXT("\tfloat3 next_vector = float3( 0.0, 0.0, 0.0 );\n");
+			OutHLSL += ReadVectorInBuffer(TEXT("next_vector"), TEXT("next_index"), TEXT("pos_col"));
+
+			OutHLSL += TEXT("\tOut_Value = lerp(prev_vector, next_vector, weight);\n");
+
+		OutHLSL += TEXT("\n}\n");
+		return true;
+	}
+	else if (DefinitionFunctionName == GetPointValueAtTimeName)
+	{
+		// GetPointValueAtTime(int In_PointID, int In_Col, float In_Time, out float Out_Value)
+		OutHLSL += TEXT("void ") + InstanceFunctionName + TEXT("(int In_PointID, int In_Col, float In_Time, out float Out_Value) \n{\n");
+		
+		OutHLSL += TEXT("\tint prev_index = -1;int next_index = -1;float weight = 1.0f;\n");
+		OutHLSL += GetRowIndexesForPointAtTime(TEXT("In_PointID"), TEXT("In_Time"), TEXT("prev_index"), TEXT("next_index"), TEXT("weight"));
+
+		OutHLSL += TEXT("\tfloat prev_value;\n");
+		OutHLSL += ReadFloatInBuffer(TEXT("prev_value"), TEXT("prev_index"), TEXT("In_Col"));
+		OutHLSL += TEXT("\tfloat next_value;\n");
+		OutHLSL += ReadFloatInBuffer(TEXT("next_value"), TEXT("next_index"), TEXT("In_Col"));
+
+		OutHLSL += TEXT("\tOut_Value = lerp(prev_value, next_value, weight);\n");
+
+		OutHLSL += TEXT("\n}\n");
+		return true;
+	}
+	else if (DefinitionFunctionName == GetPointVectorValueAtTimeName)
+	{
+		// GetPointVectorValueAtTime(int In_PointID, int In_Col, float In_Time, out float3 Out_Value)
+		OutHLSL += TEXT("void ") + InstanceFunctionName + TEXT("(int In_PointID, int In_Col, float In_Time, out float3 Out_Value) \n{\n");
+
+		OutHLSL += TEXT("\tint prev_index = -1;int next_index = -1;float weight = 1.0f;\n");
+		OutHLSL += GetRowIndexesForPointAtTime(TEXT("In_PointID"), TEXT("In_Time"), TEXT("prev_index"), TEXT("next_index"), TEXT("weight"));
+
+		OutHLSL += TEXT("\tint In_DoSwap = 1;\n");
+		OutHLSL += TEXT("\tint In_DoScale = 1;\n");
+
+		OutHLSL += TEXT("\tfloat3 prev_vector;\n");
+		OutHLSL += ReadVectorInBuffer(TEXT("prev_vector"), TEXT("prev_index"), TEXT("In_Col"));
+		OutHLSL += TEXT("\tfloat3 next_vector;\n");
+		OutHLSL += ReadVectorInBuffer(TEXT("next_vector"), TEXT("next_index"), TEXT("In_Col"));
+
+		OutHLSL += TEXT("\tOut_Value = lerp(prev_vector, next_vector, weight);\n");
+
+		OutHLSL += TEXT("\n}\n");
+		return true;
+	}
+	else if (DefinitionFunctionName == GetPointVectorValueAtTimeExName)
+	{
+		// GetPointVectorValueAtTimeEx(int In_PointID, int In_Col, float In_Time, int In_DoSwap, int In_DoScale, out float3 Out_Value)
+		OutHLSL += TEXT("void ") + InstanceFunctionName + TEXT("(int In_PointID, int In_Col, float In_Time, int In_DoSwap, int In_DoScale, out float3 Out_Value) \n{\n");
+
+			OutHLSL += TEXT("\tint prev_index = -1;int next_index = -1;float weight = 1.0f;\n");
+			OutHLSL += GetRowIndexesForPointAtTime(TEXT("In_PointID"), TEXT("In_Time"), TEXT("prev_index"), TEXT("next_index"), TEXT("weight"));
+
+			OutHLSL += TEXT("\tfloat3 prev_vector;\n");
+			OutHLSL += ReadVectorInBuffer(TEXT("prev_vector"), TEXT("prev_index"), TEXT("In_Col"));
+			OutHLSL += TEXT("\tfloat3 next_vector;\n");
+			OutHLSL += ReadVectorInBuffer(TEXT("next_vector"), TEXT("next_index"), TEXT("In_Col"));
+
+			OutHLSL += TEXT("\tOut_Value = lerp(prev_vector, next_vector, weight);\n");
+
+		OutHLSL += TEXT("\n}\n");
+
+		return true;
+	}
+	else if (DefinitionFunctionName == GetPointLifeName)
+	{
+		// GetPointLife(int In_PointID, out float Out_Value)
+		OutHLSL += TEXT("void ") + InstanceFunctionName + TEXT("(int In_PointID, out float Out_Value) \n{\n");
+
+			OutHLSL += TEXT("\tif ( ( In_PointID < 0 ) || ( In_PointID >= ") + NumberOfPointsVar + TEXT(") )\n");
+				OutHLSL += TEXT("\t\t{Out_Value = -1.0f; return;}\n");
+			OutHLSL += TEXT("\tOut_Value = ") + LifeValuesBuffer + TEXT("[ In_PointID ];\n");
+
+		OutHLSL += TEXT("\n}\n");
+		return true;
+	}
+	/*
+	else if (DefinitionFunctionName == GetPointLifeAtTimeName)
+	{
+		// GetPointLifeAtTime(int In_PointID, float In_Time, out float Out_Value)
+		OutHLSL += TEXT("void ") + InstanceFunctionName + TEXT("(int In_PointID, float In_Time, out float Out_Value) \n{\n");
+
+		OutHLSL += TEXT("\n}\n");
+		return true;
+	}
+	*/
+	else if (DefinitionFunctionName == GetPointTypeName)
+	{
+		// GetPointType(int In_PointID, out int Out_Value)
+		OutHLSL += TEXT("void ") + InstanceFunctionName + TEXT("(int In_PointID, out int Out_Value) \n{\n");
+
+			OutHLSL += TEXT("\tif ( ( In_PointID < 0 ) || ( In_PointID >= ") + NumberOfPointsVar+ TEXT(") )\n");
+				OutHLSL += TEXT("\t\t{Out_Value = -1; return;}\n");
+			OutHLSL += TEXT("\tOut_Value = ") + PointTypesBuffer + TEXT("[ In_PointID ];\n");
+
+		OutHLSL += TEXT("\n}\n");
+		return true;
+	}
+
+	return !OutHLSL.IsEmpty();
+}
+
+// Build buffers and member variables HLSL definition
+// Always use the BaseName + the DataInterfaceHLSL indentifier!
+void UNiagaraDataInterfaceHoudiniCSV::GetParameterDefinitionHLSL(FNiagaraDataInterfaceGPUParamInfo& ParamInfo, FString& OutHLSL)
+{
+	// int NumberOfRows_XX;
+	FString BufferName = UNiagaraDataInterfaceHoudiniCSV::NumberOfRowsBaseName + ParamInfo.DataInterfaceHLSLSymbol;
+	OutHLSL += TEXT("int ") + BufferName + TEXT(";\n");
+
+	// int NumberOfColumns_XX;
+	BufferName = UNiagaraDataInterfaceHoudiniCSV::NumberOfColumnsBaseName + ParamInfo.DataInterfaceHLSLSymbol;
+	OutHLSL += TEXT("int ") + BufferName + TEXT(";\n");
+
+	// int NumberOfPoints_XX;
+	BufferName = UNiagaraDataInterfaceHoudiniCSV::NumberOfPointsBaseName + ParamInfo.DataInterfaceHLSLSymbol;
+	OutHLSL += TEXT("int ") + BufferName + TEXT(";\n");
+
+	// Buffer<float> FloatValuesBuffer_XX;
+	BufferName = UNiagaraDataInterfaceHoudiniCSV::FloatValuesBufferBaseName + ParamInfo.DataInterfaceHLSLSymbol;
+	OutHLSL += TEXT("Buffer<float> ") + BufferName + TEXT(";\n");
+
+	// Buffer<int> SpecialAttributesColumnIndexesBuffer_XX;
+	BufferName = UNiagaraDataInterfaceHoudiniCSV::SpecialAttributesColumnIndexesBufferBaseName + ParamInfo.DataInterfaceHLSLSymbol;
+	OutHLSL += TEXT("Buffer<int> ") + BufferName + TEXT(";\n");
+
+	// Buffer<float> SpawnTimesBuffer_XX;
+	BufferName = UNiagaraDataInterfaceHoudiniCSV::SpawnTimesBufferBaseName + ParamInfo.DataInterfaceHLSLSymbol;
+	OutHLSL += TEXT("Buffer<float> ") + BufferName + TEXT(";\n");
+
+	// Buffer<float> LifeValuesBuffer_XX;
+	BufferName = UNiagaraDataInterfaceHoudiniCSV::LifeValuesBufferBaseName + ParamInfo.DataInterfaceHLSLSymbol;
+	OutHLSL += TEXT("Buffer<float> ") + BufferName + TEXT(";\n");
+
+	// Buffer<int> PointTypesBuffer_XX;
+	BufferName = UNiagaraDataInterfaceHoudiniCSV::PointTypesBufferBaseName + ParamInfo.DataInterfaceHLSLSymbol;
+	OutHLSL += TEXT("Buffer<int> ") + BufferName + TEXT(";\n");
+
+	// int MaxNumberOfIndexesPerPoint_XX;
+	BufferName = UNiagaraDataInterfaceHoudiniCSV::MaxNumberOfIndexesPerPointBaseName + ParamInfo.DataInterfaceHLSLSymbol;
+	OutHLSL += TEXT("int ") + BufferName + TEXT(";\n");
+
+	// Buffer<int> PointValueIndexesBuffer_XX;
+	BufferName = UNiagaraDataInterfaceHoudiniCSV::PointValueIndexesBufferBaseName + ParamInfo.DataInterfaceHLSLSymbol;
+	OutHLSL += TEXT("Buffer<int> ") + BufferName + TEXT(";\n");
+
+	// int LastSpawnedPointId_XX;
+	BufferName = UNiagaraDataInterfaceHoudiniCSV::LastSpawnedPointIdBaseName + ParamInfo.DataInterfaceHLSLSymbol;
+	OutHLSL += TEXT("int ") + BufferName + TEXT(";\n");
+
+	// int LastSpawnTime_XX;
+	BufferName = UNiagaraDataInterfaceHoudiniCSV::LastSpawnTimeBaseName + ParamInfo.DataInterfaceHLSLSymbol;
+	OutHLSL += TEXT("int ") + BufferName + TEXT(";\n");
+}
+
+FRWBuffer& UNiagaraDataInterfaceHoudiniCSV::GetFloatValuesGPUBuffer()
+{
+	//TODO: This isn't really very thread safe. Need to move to a proxy like system where DIs can push data to the RT safely.
+	if ( FloatValuesGPUBufferDirty )
+	{
+		uint32 NumElements = HoudiniCSVAsset ? HoudiniCSVAsset->FloatCSVData.Num() : 0;
+		if (NumElements <= 0)
+			return FloatValuesGPUBuffer;
+
+		FloatValuesGPUBuffer.Release();
+		FloatValuesGPUBuffer.Initialize( sizeof(float), NumElements, EPixelFormat::PF_R32_FLOAT, BUF_Static);
+
+		uint32 BufferSize = NumElements * sizeof( float );
+		float* BufferData = static_cast<float*>( RHILockVertexBuffer( FloatValuesGPUBuffer.Buffer, 0, BufferSize, EResourceLockMode::RLM_WriteOnly ) );
+
+		if ( HoudiniCSVAsset )
+			FPlatformMemory::Memcpy( BufferData, HoudiniCSVAsset->FloatCSVData.GetData(), BufferSize );
+
+		RHIUnlockVertexBuffer( FloatValuesGPUBuffer.Buffer );
+		FloatValuesGPUBufferDirty = false;
+	}
+
+	return FloatValuesGPUBuffer;
+}
+
+FRWBuffer& UNiagaraDataInterfaceHoudiniCSV::GetSpecialAttributesColumnIndexesGPUBuffer()
+{
+	//TODO: This isn't really very thread safe. Need to move to a proxy like system where DIs can push data to the RT safely.
+	if ( SpecialAttributesColumnIndexesGPUBufferDirty )
+	{
+		uint32 NumElements = HoudiniCSVAsset ? HoudiniCSVAsset->SpecialAttributesColumnIndexes.Num() : 0;
+		if (NumElements <= 0)
+			return SpecialAttributesColumnIndexesGPUBuffer;
+
+		SpecialAttributesColumnIndexesGPUBuffer.Release();
+		SpecialAttributesColumnIndexesGPUBuffer.Initialize(sizeof(int32), NumElements, EPixelFormat::PF_R32_SINT, BUF_Static);
+
+		uint32 BufferSize = NumElements * sizeof(int32);
+		float* BufferData = static_cast<float*>(RHILockVertexBuffer(SpecialAttributesColumnIndexesGPUBuffer.Buffer, 0, BufferSize, EResourceLockMode::RLM_WriteOnly));
+
+		if (HoudiniCSVAsset)
+			FPlatformMemory::Memcpy(BufferData, HoudiniCSVAsset->SpecialAttributesColumnIndexes.GetData(), BufferSize);
+
+		RHIUnlockVertexBuffer(SpecialAttributesColumnIndexesGPUBuffer.Buffer);
+		SpecialAttributesColumnIndexesGPUBufferDirty = false;
+	}
+
+	return SpecialAttributesColumnIndexesGPUBuffer;
+}
+
+FRWBuffer& UNiagaraDataInterfaceHoudiniCSV::GetSpawnTimesGPUBuffer()
+{
+	//TODO: This isn't really very thread safe. Need to move to a proxy like system where DIs can push data to the RT safely.
+	if ( SpawnTimesGPUBufferDirty )
+	{
+		uint32 NumElements = HoudiniCSVAsset ? HoudiniCSVAsset->SpawnTimes.Num() : 0;
+		if (NumElements <= 0)
+			return SpawnTimesGPUBuffer;
+
+		SpawnTimesGPUBuffer.Release();
+		SpawnTimesGPUBuffer.Initialize( sizeof(float), NumElements, EPixelFormat::PF_R32_FLOAT, BUF_Static);
+				
+		uint32 BufferSize = NumElements * sizeof( float );
+		float* BufferData = static_cast<float*>( RHILockVertexBuffer(SpawnTimesGPUBuffer.Buffer, 0, BufferSize, EResourceLockMode::RLM_WriteOnly ) );
+
+		if ( HoudiniCSVAsset )
+			FPlatformMemory::Memcpy( BufferData, HoudiniCSVAsset->SpawnTimes.GetData(), BufferSize );
+
+		RHIUnlockVertexBuffer( SpawnTimesGPUBuffer.Buffer );
+		SpawnTimesGPUBufferDirty = false;
+	}
+
+	return SpawnTimesGPUBuffer;
+}
+
+FRWBuffer& UNiagaraDataInterfaceHoudiniCSV::GetLifeValuesGPUBuffer()
+{
+	//TODO: This isn't really very thread safe. Need to move to a proxy like system where DIs can push data to the RT safely.
+	if ( LifeValuesGPUBufferDirty )
+	{
+		uint32 NumElements = HoudiniCSVAsset ? HoudiniCSVAsset->LifeValues.Num() : 0;
+		if (NumElements <= 0)
+			return LifeValuesGPUBuffer;
+
+		LifeValuesGPUBuffer.Release();
+		LifeValuesGPUBuffer.Initialize( sizeof(float), NumElements, EPixelFormat::PF_R32_FLOAT, BUF_Static);
+
+		uint32 BufferSize = NumElements * sizeof( float );
+		float* BufferData = static_cast<float*>( RHILockVertexBuffer( LifeValuesGPUBuffer.Buffer, 0, BufferSize, EResourceLockMode::RLM_WriteOnly ) );
+
+		if ( HoudiniCSVAsset )
+			FPlatformMemory::Memcpy( BufferData, HoudiniCSVAsset->LifeValues.GetData(), BufferSize );
+
+		RHIUnlockVertexBuffer(LifeValuesGPUBuffer.Buffer );
+		LifeValuesGPUBufferDirty = false;
+	}
+
+	return LifeValuesGPUBuffer;
+}
+
+FRWBuffer& UNiagaraDataInterfaceHoudiniCSV::GetPointTypesGPUBuffer()
+{
+	//TODO: This isn't really very thread safe. Need to move to a proxy like system where DIs can push data to the RT safely.
+	if ( PointTypesGPUBufferDirty )
+	{
+		uint32 NumElements = HoudiniCSVAsset ? HoudiniCSVAsset->PointTypes.Num() : 0;
+		if (NumElements <= 0)
+			return PointTypesGPUBuffer;
+
+		PointTypesGPUBuffer.Release();
+		PointTypesGPUBuffer.Initialize( sizeof(int32), NumElements, EPixelFormat::PF_R32_SINT, BUF_Static);
+
+		uint32 BufferSize = NumElements * sizeof(int32);
+		int32* BufferData = static_cast<int32*>( RHILockVertexBuffer(PointTypesGPUBuffer.Buffer, 0, BufferSize, EResourceLockMode::RLM_WriteOnly ) );
+
+		if ( HoudiniCSVAsset )
+			FPlatformMemory::Memcpy( BufferData, HoudiniCSVAsset->PointTypes.GetData(), BufferSize );
+
+		RHIUnlockVertexBuffer(PointTypesGPUBuffer.Buffer );
+		PointTypesGPUBufferDirty = false;
+	}
+
+	return PointTypesGPUBuffer;
+}
+
+FRWBuffer& UNiagaraDataInterfaceHoudiniCSV::GetPointValueIndexesGPUBuffer()
+{
+	//TODO: This isn't really very thread safe. Need to move to a proxy like system where DIs can push data to the RT safely.
+	if ( PointValueIndexesGPUBufferDirty )
+	{
+		uint32 NumPoints = HoudiniCSVAsset ? HoudiniCSVAsset->PointValueIndexes.Num() : 0;
+		// Add an extra to the max number so all indexes end by a -1
+		uint32 MaxNumIndexesPerPoint = HoudiniCSVAsset ? HoudiniCSVAsset->GetMaxNumberOfPointValueIndexes() + 1: 0;
+		uint32 NumElements = NumPoints * MaxNumIndexesPerPoint;
+
+		if (NumElements <= 0)
+			return PointValueIndexesGPUBuffer;
+
+		TArray<int32> PointValueIndexes;
+		PointValueIndexes.Init(-1, NumElements);
+		if ( HoudiniCSVAsset )
+		{
+			// We need to flatten the nested array for HLSL conversion
+			for ( int32 PointID = 0; PointID < HoudiniCSVAsset->PointValueIndexes.Num(); PointID++ )
+			{
+				TArray<int32> RowIndexes = HoudiniCSVAsset->PointValueIndexes[PointID].RowIndexes;
+				for( int32 Idx = 0; Idx < RowIndexes.Num(); Idx++ )
+				{
+					PointValueIndexes[ PointID * MaxNumIndexesPerPoint + Idx ] = RowIndexes[ Idx ];
+				}
+			}
+		}
+
+		PointValueIndexesGPUBuffer.Release();
+		PointValueIndexesGPUBuffer.Initialize(sizeof(int32), NumElements, EPixelFormat::PF_R32_SINT, BUF_Static);
+
+		uint32 BufferSize = NumElements * sizeof(int32);
+		int32* BufferData = static_cast<int32*>( RHILockVertexBuffer( PointValueIndexesGPUBuffer.Buffer, 0, BufferSize, EResourceLockMode::RLM_WriteOnly ) );
+		FPlatformMemory::Memcpy( BufferData, PointValueIndexes.GetData(), BufferSize );
+		RHIUnlockVertexBuffer( PointValueIndexesGPUBuffer.Buffer );
+
+		PointValueIndexesGPUBufferDirty = false;
+	}
+
+	return PointValueIndexesGPUBuffer;
+}
+
+// Parameters used for GPU sim compatibility
+struct FNiagaraDataInterfaceParametersCS_HoudiniCSV : public FNiagaraDataInterfaceParametersCS
+{
+	virtual void Bind(const FNiagaraDataInterfaceParamRef& ParamRef, const class FShaderParameterMap& ParameterMap) override
+	{
+		NumberOfRows.Bind(ParameterMap, *(UNiagaraDataInterfaceHoudiniCSV::NumberOfRowsBaseName + ParamRef.ParameterInfo.DataInterfaceHLSLSymbol));
+		NumberOfColumns.Bind(ParameterMap, *(UNiagaraDataInterfaceHoudiniCSV::NumberOfColumnsBaseName + ParamRef.ParameterInfo.DataInterfaceHLSLSymbol));
+		NumberOfPoints.Bind(ParameterMap, *(UNiagaraDataInterfaceHoudiniCSV::NumberOfPointsBaseName + ParamRef.ParameterInfo.DataInterfaceHLSLSymbol));
+		
+		FloatValuesBuffer.Bind(ParameterMap, *(UNiagaraDataInterfaceHoudiniCSV::FloatValuesBufferBaseName + ParamRef.ParameterInfo.DataInterfaceHLSLSymbol));
+
+		SpecialAttributesColumnIndexesBuffer.Bind(ParameterMap, *(UNiagaraDataInterfaceHoudiniCSV::SpecialAttributesColumnIndexesBufferBaseName + ParamRef.ParameterInfo.DataInterfaceHLSLSymbol));		
+
+		SpawnTimesBuffer.Bind(ParameterMap, *(UNiagaraDataInterfaceHoudiniCSV::SpawnTimesBufferBaseName + ParamRef.ParameterInfo.DataInterfaceHLSLSymbol));
+		LifeValuesBuffer.Bind(ParameterMap, *(UNiagaraDataInterfaceHoudiniCSV::LifeValuesBufferBaseName + ParamRef.ParameterInfo.DataInterfaceHLSLSymbol));
+		PointTypesBuffer.Bind(ParameterMap, *(UNiagaraDataInterfaceHoudiniCSV::PointTypesBufferBaseName + ParamRef.ParameterInfo.DataInterfaceHLSLSymbol));
+
+		MaxNumberOfIndexesPerPoint.Bind(ParameterMap, *(UNiagaraDataInterfaceHoudiniCSV::MaxNumberOfIndexesPerPointBaseName + ParamRef.ParameterInfo.DataInterfaceHLSLSymbol));
+		PointValueIndexesBuffer.Bind(ParameterMap, *(UNiagaraDataInterfaceHoudiniCSV::PointValueIndexesBufferBaseName + ParamRef.ParameterInfo.DataInterfaceHLSLSymbol));
+
+		LastSpawnedPointId.Bind(ParameterMap, *(UNiagaraDataInterfaceHoudiniCSV::LastSpawnedPointIdBaseName + ParamRef.ParameterInfo.DataInterfaceHLSLSymbol));
+		LastSpawnTime.Bind(ParameterMap, *(UNiagaraDataInterfaceHoudiniCSV::LastSpawnTimeBaseName + ParamRef.ParameterInfo.DataInterfaceHLSLSymbol));
+	}
+
+	virtual void Serialize(FArchive& Ar)override
+	{
+		Ar << Version;
+
+		Ar << NumberOfRows;
+		Ar << NumberOfColumns;
+		Ar << NumberOfPoints;
+
+		Ar << FloatValuesBuffer;
+		Ar << SpecialAttributesColumnIndexesBuffer;
+
+		Ar << SpawnTimesBuffer;
+		Ar << LifeValuesBuffer;
+		Ar << PointTypesBuffer;
+
+		Ar << MaxNumberOfIndexesPerPoint;
+		Ar << PointValueIndexesBuffer;
+
+		Ar << LastSpawnedPointId;
+		Ar << LastSpawnTime;
+	}
+
+	virtual void Set( FRHICommandList& RHICmdList, FNiagaraShader* Shader, class UNiagaraDataInterface* DataInterface ) const override
+	{
+		check( IsInRenderingThread() );
+
+		const FComputeShaderRHIParamRef ComputeShaderRHI = Shader->GetComputeShader();
+		UNiagaraDataInterfaceHoudiniCSV* HoudiniDI = CastChecked<UNiagaraDataInterfaceHoudiniCSV>( DataInterface );
+		if ( !HoudiniDI )
+			return;
+
+		SetShaderValue(RHICmdList, ComputeShaderRHI, NumberOfRows, HoudiniDI->GetNumberOfRows());
+		SetShaderValue(RHICmdList, ComputeShaderRHI, NumberOfColumns, HoudiniDI->GetNumberOfColumns());
+		SetShaderValue(RHICmdList, ComputeShaderRHI, NumberOfPoints, HoudiniDI->GetNumberOfPoints());
+
+		FRWBuffer& FloatRWBuffer = HoudiniDI->GetFloatValuesGPUBuffer();
+ 		RHICmdList.SetShaderResourceViewParameter( ComputeShaderRHI, FloatValuesBuffer.GetBaseIndex(), FloatRWBuffer.SRV );
+
+		FRWBuffer& AttributeColumnRWBuffer = HoudiniDI->GetSpecialAttributesColumnIndexesGPUBuffer();
+		RHICmdList.SetShaderResourceViewParameter(ComputeShaderRHI, SpecialAttributesColumnIndexesBuffer.GetBaseIndex(), AttributeColumnRWBuffer.SRV);
+
+		FRWBuffer& SpawnRWBuffer = HoudiniDI->GetSpawnTimesGPUBuffer();
+		RHICmdList.SetShaderResourceViewParameter(ComputeShaderRHI, SpawnTimesBuffer.GetBaseIndex(), SpawnRWBuffer.SRV);
+		FRWBuffer& LifeRWBuffer = HoudiniDI->GetLifeValuesGPUBuffer();
+		RHICmdList.SetShaderResourceViewParameter(ComputeShaderRHI, LifeValuesBuffer.GetBaseIndex(), LifeRWBuffer.SRV);
+		FRWBuffer& TypesRWBuffer = HoudiniDI->GetPointTypesGPUBuffer();
+		RHICmdList.SetShaderResourceViewParameter(ComputeShaderRHI, PointTypesBuffer.GetBaseIndex(), TypesRWBuffer.SRV);
+
+		SetShaderValue(RHICmdList, ComputeShaderRHI, MaxNumberOfIndexesPerPoint, HoudiniDI->GetMaxNumberOfIndexesPerPoints());
+
+		FRWBuffer& PointValuesIndexesRWBuffer = HoudiniDI->GetPointValueIndexesGPUBuffer();
+		RHICmdList.SetShaderResourceViewParameter(ComputeShaderRHI, PointValueIndexesBuffer.GetBaseIndex(), PointValuesIndexesRWBuffer.SRV);
+
+		SetShaderValue(RHICmdList, ComputeShaderRHI, LastSpawnedPointId, -1);
+		SetShaderValue(RHICmdList, ComputeShaderRHI, LastSpawnTime, -1);
+	}
+
+private:
+
+	FShaderParameter NumberOfRows;
+	FShaderParameter NumberOfColumns;
+	FShaderParameter NumberOfPoints;
+
+	FShaderResourceParameter FloatValuesBuffer;
+	FShaderResourceParameter SpecialAttributesColumnIndexesBuffer;
+
+	FShaderResourceParameter SpawnTimesBuffer;
+	FShaderResourceParameter LifeValuesBuffer;
+	FShaderResourceParameter PointTypesBuffer;
+
+	FShaderParameter MaxNumberOfIndexesPerPoint;
+	FShaderResourceParameter PointValueIndexesBuffer;
+
+	FShaderParameter LastSpawnedPointId;
+	FShaderParameter LastSpawnTime;
+
+	uint32 Version = 1;
+};
+
+FNiagaraDataInterfaceParametersCS* UNiagaraDataInterfaceHoudiniCSV::ConstructComputeParameters()const
+{
+	return new FNiagaraDataInterfaceParametersCS_HoudiniCSV();
+}
+
 #undef LOCTEXT_NAMESPACE
diff --git a/Engine/Plugins/FX/HoudiniNiagara/Source/HoudiniNiagara/Public/HoudiniCSV.h b/Engine/Plugins/FX/HoudiniNiagara/Source/HoudiniNiagara/Public/HoudiniCSV.h
--- a/Engine/Plugins/FX/HoudiniNiagara/Source/HoudiniNiagara/Public/HoudiniCSV.h
+++ b/Engine/Plugins/FX/HoudiniNiagara/Source/HoudiniNiagara/Public/HoudiniCSV.h
@@ -60,7 +60,7 @@ struct FPointIndexes
 	TArray<int32> RowIndexes;
 };
 
-UCLASS()
+UCLASS(BlueprintType)
 class HOUDININIAGARA_API UHoudiniCSV : public UObject
 {
 	friend class UNiagaraDataInterfaceHoudiniCSV;
@@ -79,74 +79,104 @@ class HOUDININIAGARA_API UHoudiniCSV : public UObject
 	void SetFileName( const FString& TheFilename );
 
 	// Returns the number of points found in the CSV file
+	UFUNCTION(BlueprintCallable, Category = "Houdini Attributes Data")
 	int32 GetNumberOfPoints() const;
 
 	// Returns the number of rows found in the CSV file
+	UFUNCTION(BlueprintCallable, Category = "Houdini Attributes Data")
 	int32 GetNumberOfRows() const ;
 
 	// Returns the number of columns found in the CSV file
+	UFUNCTION(BlueprintCallable, Category = "Houdini Attributes Data")
 	int32 GetNumberOfColumns() const;
 
 	// Return the column index for a specific attribute
+	//UFUNCTION(BlueprintCallable, Category = "Houdini Attributes Data")
 	int32 GetAttributeColumnIndex(const EHoudiniAttributes& Attr) const;
 
 	// Returns if the specific attribute has a valid column index
 	bool IsValidAttributeColumnIndex(const EHoudiniAttributes& Attr) const;
 
 	// Returns the column index for a given string
+	UFUNCTION(BlueprintCallable, Category = "Houdini Attributes Data")
 	bool GetColumnIndexFromString(const FString& ColumnTitle, int32& ColumnIndex) const;
 
 	// Returns the float value at a given point in the CSV file
+	UFUNCTION(BlueprintCallable, Category = "Houdini Attributes Data")
 	bool GetFloatValue( const int32& rowIndex, const int32& colIndex, float& value ) const;
 	// Returns the float value at a given point in the CSV file
-	bool GetFloatValue( const int32& rowIndex, const FString& ColumnTitle, float& value ) const;
+	UFUNCTION(BlueprintCallable, Category = "Houdini Attributes Data")
+	bool GetFloatValueForString( const int32& rowIndex, const FString& ColumnTitle, float& value ) const;
 	/*
-	// Returns the float value at a given point in the CSV file
+	// Returns the string value at a given point in the CSV file
 	bool GetCSVStringValue( const int32& rowIndex, const int32& colIndex, FString& value );
 	// Returns the string value at a given point in the CSV file
 	bool GetCSVStringValue( const int32& rowIndex, const FString& ColumnTitle, FString& value );
 	*/
 	// Returns a Vector3 for a given point in the CSV file
+	UFUNCTION(BlueprintCallable, Category = "Houdini Attributes Data")
 	bool GetVectorValue( const int32& rowIndex, const int32& colIndex, FVector& value, const bool& DoSwap = true, const bool& DoScale = true ) const;
 	
 	// Returns a time value for a given point in the CSV file
+	UFUNCTION(BlueprintCallable, Category = "Houdini Attributes Data")
 	bool GetTimeValue( const int32& rowIndex, float& value ) const;
 	// Returns a Position Vector3 for a given point in the CSV file (converted to unreal's coordinate system)
+	UFUNCTION(BlueprintCallable, Category = "Houdini Attributes Data")
 	bool GetPositionValue( const int32& rowIndex, FVector& value ) const;
 	// Returns a Normal Vector3 for a given point in the CSV file (converted to unreal's coordinate system)
+	UFUNCTION(BlueprintCallable, Category = "Houdini Attributes Data")
 	bool GetNormalValue( const int32& rowIndex, FVector& value ) const;
 	// Returns a Color for a given point in the CSV file
+	UFUNCTION(BlueprintCallable, Category = "Houdini Attributes Data")
 	bool GetColorValue( const int32& rowIndex, FLinearColor& value ) const;
 	// Returns a Velocity Vector3 for a given point in the CSV file
+	UFUNCTION(BlueprintCallable, Category = "Houdini Attributes Data")
 	bool GetVelocityValue(const int32& rowIndex, FVector& value ) const;
 
 	// Get the last row index for a given time value (the row with a time smaller or equal to desiredTime)
 	// If the CSV file doesn't have time informations, returns false and set LastRowIndex to the last row in the file
 	// If desiredTime is smaller than the time value in the first row, LastRowIndex will be set to -1
 	// If desiredTime is higher than the last time value in the last row of the csv file, LastIndex will be set to the last row's index
+	UFUNCTION(BlueprintCallable, Category = "Houdini Attributes Data")
 	bool GetLastRowIndexAtTime( const float& desiredTime, int32& lastRowIndex ) const;
 
 	// Get the last pointID of the points to be spawned at time t
 	// Invalid Index are used to indicate edge cases:
 	// -1 will be returned if there is no points to spawn ( t is smaller than the first point time )
 	// NumberOfRows will be returned if all points in the CSV have been spawned ( t is higher than the last point time )
+	UFUNCTION(BlueprintCallable, Category = "Houdini Attributes Data")
 	bool GetLastPointIDToSpawnAtTime( const float& time, int32& lastIndex ) const;
+	UFUNCTION(BlueprintCallable, Category = "Houdini Attributes Data")
+	bool GetPointIDsToSpawnAtTime(
+		const float& desiredTime,
+		int32& MinID, int32& MaxID, int32& Count,
+		int32& LastSpawnedPointID, float& LastSpawnTime ) const;
 
 	// Returns the previous and next indexes for reading the values of a specified point at a given time
+	UFUNCTION(BlueprintCallable, Category = "Houdini Attributes Data")
 	bool GetRowIndexesForPointAtTime(const int32& PointID, const float& desiredTime, int32& PrevIndex, int32& NextIndex, float& PrevWeight) const;
-	// Returns the value for a point at a given time value (linearly interpolated) 
+	// Returns the value for a point at a given time value (linearly interpolated)
+	UFUNCTION(BlueprintCallable, Category = "Houdini Attributes Data")
 	bool GetPointValueAtTime(const int32& PointID, const int32& ColumnIndex, const float& desiredTime, float& Value) const;
-	// Returns the Vector Value for a given point at a given time value (linearly interpolated) 
+	// Returns the Vector Value for a given point at a given time value (linearly interpolated)
+	UFUNCTION(BlueprintCallable, Category = "Houdini Attributes Data")
 	bool GetPointVectorValueAtTime(const int32& PointID, const int32& ColumnIndex, const float& desiredTime, FVector& Vector, const bool& DoSwap, const bool& DoScale) const;
-	// Returns the Position Value for a given point at a given time value (linearly interpolated) 
+	// Returns the Position Value for a given point at a given time value (linearly interpolated)
+	UFUNCTION(BlueprintCallable, Category = "Houdini Attributes Data")
 	bool GetPointPositionAtTime(const int32& PointID, const float& desiredTime, FVector& Vector) const;
 	// Return a given point's life value at spawn
+	UFUNCTION(BlueprintCallable, Category = "Houdini Attributes Data")
 	bool GetPointLife(const int32& PointID, float& Value) const;
 	// Return a point's life for a given time value
+	UFUNCTION(BlueprintCallable, Category = "Houdini Attributes Data")
 	bool GetPointLifeAtTime(const int32& PointID, const float& DesiredTime, float& Value) const;
 	// Return a point's type at spawn
+	UFUNCTION(BlueprintCallable, Category = "Houdini Attributes Data")
 	bool GetPointType(const int32& PointID, int32& Value) const;
 
+	// Returns the maximum number of indexes per point, used for flattening the buffer for HLSL conversion
+	int32 GetMaxNumberOfPointValueIndexes() const;
+
 	//-----------------------------------------------------------------------------------------
 	//  MEMBER VARIABLES
 	//-----------------------------------------------------------------------------------------
diff --git a/Engine/Plugins/FX/HoudiniNiagara/Source/HoudiniNiagaraEditor/Private/HoudiniNiagaraEditor.cpp b/Engine/Plugins/FX/HoudiniNiagara/Source/HoudiniNiagaraEditor/Private/HoudiniNiagaraEditor.cpp
--- a/Engine/Plugins/FX/HoudiniNiagara/Source/HoudiniNiagaraEditor/Private/HoudiniNiagaraEditor.cpp
+++ b/Engine/Plugins/FX/HoudiniNiagara/Source/HoudiniNiagaraEditor/Private/HoudiniNiagaraEditor.cpp
@@ -24,6 +24,8 @@
 #include "HoudiniNiagaraEditor.h"
 #include "HoudiniCSVAssetActions.h"
 #include "AssetRegistryModule.h"
+#include "Styling/SlateStyleRegistry.h"
+#include "Styling/SlateStyle.h"
 
 #define LOCTEXT_NAMESPACE "FHoudiniNiagaraModule"
 
@@ -35,6 +37,34 @@ void FHoudiniNiagaraEditorModule::StartupModule()
     TSharedRef< IAssetTypeActions > HCSVAction = MakeShareable( new FHoudiniCSVAssetActions() );
     AssetTools.RegisterAssetTypeActions( HCSVAction );
     AssetTypeActions.Add( HCSVAction );
+
+	// Create Slate style set.
+	if (!StyleSet.IsValid())
+	{
+		// Create Slate style set.	
+		StyleSet = MakeShareable(new FSlateStyleSet(TEXT("HoudiniNiagaraStyle")));
+		StyleSet->SetContentRoot(FPaths::EngineContentDir() / TEXT("Editor/Slate"));
+		StyleSet->SetCoreContentRoot(FPaths::EngineContentDir() / TEXT("Slate"));
+
+		// Note, these sizes are in Slate Units. Slate Units do NOT have to map to pixels.
+		const FVector2D Icon16x16(16.0f, 16.0f);
+		const FVector2D Icon64x64(64.0f, 64.0f);
+		const FVector2D Icon128x128(128.0f, 128.0f);
+		
+		static FString IconsDir = FPaths::EnginePluginsDir() / TEXT("FX/HoudiniNiagara/Resources/");
+
+		// Register the Asset icon
+		StyleSet->Set(
+			"ClassIcon.HoudiniCSV",
+			new FSlateImageBrush(IconsDir + TEXT("HCSVIcon128.png"), Icon16x16));
+
+		StyleSet->Set(
+			"ClassThumbnail.HoudiniCSV",
+			new FSlateImageBrush(IconsDir + TEXT("HCSVIcon128.png"), Icon64x64));
+
+		// Register Slate style.
+		FSlateStyleRegistry::RegisterSlateStyle(*StyleSet.Get());
+	}
 }
 
 void FHoudiniNiagaraEditorModule::ShutdownModule()
@@ -49,6 +79,16 @@ void FHoudiniNiagaraEditorModule::ShutdownModule()
 
 		AssetTypeActions.Empty();
     }
+
+	// Unregister Slate style set.
+	if (StyleSet.IsValid())
+	{
+		// Unregister Slate style.
+		FSlateStyleRegistry::UnRegisterSlateStyle(*StyleSet.Get());
+
+		ensure(StyleSet.IsUnique());
+		StyleSet.Reset();
+	}
 }
 
 #undef LOCTEXT_NAMESPACE
diff --git a/Engine/Plugins/FX/HoudiniNiagara/Source/HoudiniNiagaraEditor/Public/HoudiniNiagaraEditor.h b/Engine/Plugins/FX/HoudiniNiagara/Source/HoudiniNiagaraEditor/Public/HoudiniNiagaraEditor.h
--- a/Engine/Plugins/FX/HoudiniNiagara/Source/HoudiniNiagaraEditor/Public/HoudiniNiagaraEditor.h
+++ b/Engine/Plugins/FX/HoudiniNiagara/Source/HoudiniNiagaraEditor/Public/HoudiniNiagaraEditor.h
@@ -27,6 +27,7 @@
 #include "Modules/ModuleInterface.h"
 
 class IAssetTypeActions;
+class FSlateStyleSet;
 
 class FHoudiniNiagaraEditorModule : public IModuleInterface
 {
@@ -40,4 +41,7 @@ class FHoudiniNiagaraEditorModule : public IModuleInterface
 
 	/** AssetType actions associated with Houdini CSV assets. **/
 	TArray< TSharedPtr< IAssetTypeActions > > AssetTypeActions;
+
+	/** Slate styleset used by this module. **/
+	TSharedPtr< FSlateStyleSet > StyleSet;
 };
\ No newline at end of file
diff --git a/Engine/Plugins/FX/Niagara/Source/Niagara/Classes/NiagaraDataInterfaceSpline.h b/Engine/Plugins/FX/Niagara/Source/Niagara/Classes/NiagaraDataInterfaceSpline.h
--- a/Engine/Plugins/FX/Niagara/Source/Niagara/Classes/NiagaraDataInterfaceSpline.h
+++ b/Engine/Plugins/FX/Niagara/Source/Niagara/Classes/NiagaraDataInterfaceSpline.h
@@ -19,7 +19,7 @@ struct FNDISpline_InstanceData
 	FMatrix TransformInverseTransposed;
 };
 
-/** Data Interface allowing sampling of static meshes. */
+/** Data Interface allowing sampling of in-world spline components. Note that this data interface is very experimental. */
 UCLASS(EditInlineNew, Category = "Splines", meta = (DisplayName = "Spline"))
 class NIAGARA_API UNiagaraDataInterfaceSpline : public UNiagaraDataInterface
 {
diff --git a/Engine/Plugins/FX/Niagara/Source/Niagara/Private/NiagaraDataInterfaceSpline.cpp b/Engine/Plugins/FX/Niagara/Source/Niagara/Private/NiagaraDataInterfaceSpline.cpp
--- a/Engine/Plugins/FX/Niagara/Source/Niagara/Private/NiagaraDataInterfaceSpline.cpp
+++ b/Engine/Plugins/FX/Niagara/Source/Niagara/Private/NiagaraDataInterfaceSpline.cpp
@@ -398,22 +398,42 @@ void UNiagaraDataInterfaceSpline::SampleSplinePositionByUnitDistance(FVectorVMCo
 	FRegisterHandler<float> OutPosY(Context);
 	FRegisterHandler<float> OutPosZ(Context);
 
-	checkSlow(InstData->Component);
-
-	for (int32 i = 0; i < Context.NumInstances; ++i)
+	if (InstData->Component)
 	{
-		float DistanceUnitDistance = SplineSampleParam.Get();
 
-		FVector Pos = InstData->Component->GetLocationAtDistanceAlongSpline(DistanceUnitDistance * InstData->Component->GetSplineLength(), ESplineCoordinateSpace::Local);
-		TransformHandler.TransformPosition(Pos, InstData->Transform);
-
-		*OutPosX.GetDest() = Pos.X;
-		*OutPosY.GetDest() = Pos.Y;
-		*OutPosZ.GetDest() = Pos.Z;
-		SplineSampleParam.Advance();
-		OutPosX.Advance();
-		OutPosY.Advance();
-		OutPosZ.Advance();
+		for (int32 i = 0; i < Context.NumInstances; ++i)
+		{
+			float DistanceUnitDistance = SplineSampleParam.Get();
+
+			FVector Pos = InstData->Component->GetLocationAtDistanceAlongSpline(DistanceUnitDistance * InstData->Component->GetSplineLength(), ESplineCoordinateSpace::Local);
+			TransformHandler.TransformPosition(Pos, InstData->Transform);
+
+			*OutPosX.GetDest() = Pos.X;
+			*OutPosY.GetDest() = Pos.Y;
+			*OutPosZ.GetDest() = Pos.Z;
+			SplineSampleParam.Advance();
+			OutPosX.Advance();
+			OutPosY.Advance();
+			OutPosZ.Advance();
+		}
+	}
+	else
+	{
+		for (int32 i = 0; i < Context.NumInstances; ++i)
+		{
+			float DistanceUnitDistance = SplineSampleParam.Get();
+
+			FVector Pos = FVector(EForceInit::ForceInitToZero);
+			TransformHandler.TransformPosition(Pos, InstData->Transform);
+
+			*OutPosX.GetDest() = Pos.X;
+			*OutPosY.GetDest() = Pos.Y;
+			*OutPosZ.GetDest() = Pos.Z;
+			SplineSampleParam.Advance();
+			OutPosX.Advance();
+			OutPosY.Advance();
+			OutPosZ.Advance();
+		}
 	}
 }
 
@@ -427,22 +447,42 @@ void UNiagaraDataInterfaceSpline::SampleSplineUpVectorByUnitDistance(FVectorVMCo
 	FRegisterHandler<float> OutPosY(Context);
 	FRegisterHandler<float> OutPosZ(Context);
 
-	checkSlow(InstData->Component);
-
-	for (int32 i = 0; i < Context.NumInstances; ++i)
+	if (InstData->Component)
 	{
-		float DistanceUnitDistance = SplineSampleParam.Get();
 
-		FVector Pos = InstData->Component->GetUpVectorAtDistanceAlongSpline(DistanceUnitDistance * InstData->Component->GetSplineLength(), ESplineCoordinateSpace::Local);
-		TransformHandler.TransformPosition(Pos, InstData->Transform);
-
-		*OutPosX.GetDest() = Pos.X;
-		*OutPosY.GetDest() = Pos.Y;
-		*OutPosZ.GetDest() = Pos.Z;
-		SplineSampleParam.Advance();
-		OutPosX.Advance();
-		OutPosY.Advance();
-		OutPosZ.Advance();
+		for (int32 i = 0; i < Context.NumInstances; ++i)
+		{
+			float DistanceUnitDistance = SplineSampleParam.Get();
+
+			FVector Pos = InstData->Component->GetUpVectorAtDistanceAlongSpline(DistanceUnitDistance * InstData->Component->GetSplineLength(), ESplineCoordinateSpace::Local);
+			TransformHandler.TransformVector(Pos, InstData->Transform);
+
+			*OutPosX.GetDest() = Pos.X;
+			*OutPosY.GetDest() = Pos.Y;
+			*OutPosZ.GetDest() = Pos.Z;
+			SplineSampleParam.Advance();
+			OutPosX.Advance();
+			OutPosY.Advance();
+			OutPosZ.Advance();
+		}
+	}
+	else
+	{
+		for (int32 i = 0; i < Context.NumInstances; ++i)
+		{
+			float DistanceUnitDistance = SplineSampleParam.Get();
+
+			FVector Pos = FVector(0.0f, 0.0f, 1.0f); 
+			TransformHandler.TransformVector(Pos, InstData->Transform);
+
+			*OutPosX.GetDest() = Pos.X;
+			*OutPosY.GetDest() = Pos.Y;
+			*OutPosZ.GetDest() = Pos.Z;
+			SplineSampleParam.Advance();
+			OutPosX.Advance();
+			OutPosY.Advance();
+			OutPosZ.Advance();
+		}
 	}
 }
 
@@ -455,22 +495,41 @@ void UNiagaraDataInterfaceSpline::SampleSplineRightVectorByUnitDistance(FVectorV
 	FRegisterHandler<float> OutPosX(Context);
 	FRegisterHandler<float> OutPosY(Context);
 	FRegisterHandler<float> OutPosZ(Context);
-	checkSlow(InstData->Component);
-
-	for (int32 i = 0; i < Context.NumInstances; ++i)
+	
+	
+	if (InstData->Component)
 	{
-		float DistanceUnitDistance = SplineSampleParam.Get();
-
-		FVector Pos = InstData->Component->GetRightVectorAtDistanceAlongSpline(DistanceUnitDistance * InstData->Component->GetSplineLength(), ESplineCoordinateSpace::Local);
-		TransformHandler.TransformPosition(Pos, InstData->Transform);
-
-		*OutPosX.GetDest() = Pos.X;
-		*OutPosY.GetDest() = Pos.Y;
-		*OutPosZ.GetDest() = Pos.Z;
-		SplineSampleParam.Advance();
-		OutPosX.Advance();
-		OutPosY.Advance();
-		OutPosZ.Advance();
+		for (int32 i = 0; i < Context.NumInstances; ++i)
+		{
+			float DistanceUnitDistance = SplineSampleParam.Get();
+
+			FVector Pos = InstData->Component->GetRightVectorAtDistanceAlongSpline(DistanceUnitDistance * InstData->Component->GetSplineLength(), ESplineCoordinateSpace::Local);
+			TransformHandler.TransformVector(Pos, InstData->Transform);
+
+			*OutPosX.GetDest() = Pos.X;
+			*OutPosY.GetDest() = Pos.Y;
+			*OutPosZ.GetDest() = Pos.Z;
+			SplineSampleParam.Advance();
+			OutPosX.Advance();
+			OutPosY.Advance();
+			OutPosZ.Advance();
+		}
+	}
+	else
+	{
+		for (int32 i = 0; i < Context.NumInstances; ++i)
+		{
+			FVector Pos = FVector(-1.0f, 0.0f, 0.0f); 
+			TransformHandler.TransformVector(Pos, InstData->Transform);
+
+			*OutPosX.GetDest() = Pos.X;
+			*OutPosY.GetDest() = Pos.Y;
+			*OutPosZ.GetDest() = Pos.Z;
+			SplineSampleParam.Advance();
+			OutPosX.Advance();
+			OutPosY.Advance();
+			OutPosZ.Advance();
+		}
 	}
 }
 
@@ -484,21 +543,41 @@ void UNiagaraDataInterfaceSpline::SampleSplineTangentByUnitDistance(FVectorVMCon
 	FRegisterHandler<float> OutPosY(Context);
 	FRegisterHandler<float> OutPosZ(Context);
 
-	checkSlow(InstData->Component);
-	for (int32 i = 0; i < Context.NumInstances; ++i)
+	if (InstData->Component)
 	{
-		float DistanceUnitDistance = SplineSampleParam.Get();
-
-		FVector Pos = InstData->Component->GetTangentAtDistanceAlongSpline(DistanceUnitDistance * InstData->Component->GetSplineLength(), ESplineCoordinateSpace::Local);
-		TransformHandler.TransformPosition(Pos, InstData->Transform);
-
-		*OutPosX.GetDest() = Pos.X;
-		*OutPosY.GetDest() = Pos.Y;
-		*OutPosZ.GetDest() = Pos.Z;
-		SplineSampleParam.Advance();
-		OutPosX.Advance();
-		OutPosY.Advance();
-		OutPosZ.Advance();
+		for (int32 i = 0; i < Context.NumInstances; ++i)
+		{
+			float DistanceUnitDistance = SplineSampleParam.Get();
+
+			FVector Pos = InstData->Component->GetTangentAtDistanceAlongSpline(DistanceUnitDistance * InstData->Component->GetSplineLength(), ESplineCoordinateSpace::Local);
+			TransformHandler.TransformVector(Pos, InstData->Transform);
+
+			*OutPosX.GetDest() = Pos.X;
+			*OutPosY.GetDest() = Pos.Y;
+			*OutPosZ.GetDest() = Pos.Z;
+			SplineSampleParam.Advance();
+			OutPosX.Advance();
+			OutPosY.Advance();
+			OutPosZ.Advance();
+		}
+	}
+	else
+	{
+		for (int32 i = 0; i < Context.NumInstances; ++i)
+		{
+			float DistanceUnitDistance = SplineSampleParam.Get();
+
+			FVector Pos = FVector(EForceInit::ForceInitToZero); 
+			TransformHandler.TransformVector(Pos, InstData->Transform);
+
+			*OutPosX.GetDest() = Pos.X;
+			*OutPosY.GetDest() = Pos.Y;
+			*OutPosZ.GetDest() = Pos.Z;
+			SplineSampleParam.Advance();
+			OutPosX.Advance();
+			OutPosY.Advance();
+			OutPosZ.Advance();
+		}
 	}
 }
 
@@ -512,22 +591,41 @@ void UNiagaraDataInterfaceSpline::SampleSplineDirectionByUnitDistance(FVectorVMC
 	FRegisterHandler<float> OutPosY(Context);
 	FRegisterHandler<float> OutPosZ(Context);
 
-	checkSlow(InstData->Component);
-
-	for (int32 i = 0; i < Context.NumInstances; ++i)
+	if (InstData->Component)
 	{
-		float DistanceUnitDistance = SplineSampleParam.Get();
-
-		FVector Pos = InstData->Component->GetDirectionAtDistanceAlongSpline(DistanceUnitDistance * InstData->Component->GetSplineLength(), ESplineCoordinateSpace::Local);
-		TransformHandler.TransformPosition(Pos, InstData->Transform);
-
-		*OutPosX.GetDest() = Pos.X;
-		*OutPosY.GetDest() = Pos.Y;
-		*OutPosZ.GetDest() = Pos.Z;
-		SplineSampleParam.Advance();
-		OutPosX.Advance();
-		OutPosY.Advance();
-		OutPosZ.Advance();
+		for (int32 i = 0; i < Context.NumInstances; ++i)
+		{
+			float DistanceUnitDistance = SplineSampleParam.Get();
+
+			FVector Pos = InstData->Component->GetDirectionAtDistanceAlongSpline(DistanceUnitDistance * InstData->Component->GetSplineLength(), ESplineCoordinateSpace::Local);
+			TransformHandler.TransformVector(Pos, InstData->Transform);
+
+			*OutPosX.GetDest() = Pos.X;
+			*OutPosY.GetDest() = Pos.Y;
+			*OutPosZ.GetDest() = Pos.Z;
+			SplineSampleParam.Advance();
+			OutPosX.Advance();
+			OutPosY.Advance();
+			OutPosZ.Advance();
+		}
+	}
+	else
+	{
+		for (int32 i = 0; i < Context.NumInstances; ++i)
+		{
+			float DistanceUnitDistance = SplineSampleParam.Get();
+
+			FVector Pos = FVector(0.0f, 1.0f, 0.0f); 
+			TransformHandler.TransformVector(Pos, InstData->Transform);
+
+			*OutPosX.GetDest() = Pos.X;
+			*OutPosY.GetDest() = Pos.Y;
+			*OutPosZ.GetDest() = Pos.Z;
+			SplineSampleParam.Advance();
+			OutPosX.Advance();
+			OutPosY.Advance();
+			OutPosZ.Advance();
+		}
 	}
 }
 
@@ -580,30 +678,48 @@ void UNiagaraDataInterfaceSpline::FindClosestUnitDistanceFromPositionWS(FVectorV
 	FUserPtrHandler<FNDISpline_InstanceData> InstData(Context);
 	FRegisterHandler<float> OutUnitDistance(Context);
 
-	checkSlow(InstData->Component);
+	if (InstData->Component)
+	{
 
-	const int32 NumPoints = InstData->Component->GetSplinePointsPosition().Points.Num();
-	const float FinalKeyTime = InstData->Component->GetSplinePointsPosition().Points[NumPoints - 1].InVal;
+		const int32 NumPoints = InstData->Component->GetSplinePointsPosition().Points.Num();
+		const float FinalKeyTime = InstData->Component->GetSplinePointsPosition().Points[NumPoints - 1].InVal;
 
-	for (int32 i = 0; i < Context.NumInstances; ++i)
-	{
-		float PosX = PosXParam.Get();
-		float PosY = PosYParam.Get();
-		float PosZ = PosZParam.Get();
+		for (int32 i = 0; i < Context.NumInstances; ++i)
+		{
+			float PosX = PosXParam.Get();
+			float PosY = PosYParam.Get();
+			float PosZ = PosZParam.Get();
 
-		FVector Pos(PosX, PosY, PosZ);
+			FVector Pos(PosX, PosY, PosZ);
 
-		// This first call finds the key time, but this is not in 0..1 range for the spline. 
-		float KeyTime = InstData->Component->FindInputKeyClosestToWorldLocation(Pos);
-		// We need to convert into the range by dividing through by the overall duration of the spline according to the keys.
-		float UnitDistance = KeyTime / FinalKeyTime;
+			// This first call finds the key time, but this is not in 0..1 range for the spline. 
+			float KeyTime = InstData->Component->FindInputKeyClosestToWorldLocation(Pos);
+			// We need to convert into the range by dividing through by the overall duration of the spline according to the keys.
+			float UnitDistance = KeyTime / FinalKeyTime;
 
-		*OutUnitDistance.GetDest() = UnitDistance;
+			*OutUnitDistance.GetDest() = UnitDistance;
 
-		PosXParam.Advance();
-		PosYParam.Advance();
-		PosZParam.Advance();
-		OutUnitDistance.Advance();
+			PosXParam.Advance();
+			PosYParam.Advance();
+			PosZParam.Advance();
+			OutUnitDistance.Advance();
+		}
+	}
+	else
+	{
+		for (int32 i = 0; i < Context.NumInstances; ++i)
+		{
+			float PosX = PosXParam.Get();
+			float PosY = PosYParam.Get();
+			float PosZ = PosZParam.Get();
+
+			*OutUnitDistance.GetDest() = 0.0f;
+
+			PosXParam.Advance();
+			PosYParam.Advance();
+			PosZParam.Advance();
+			OutUnitDistance.Advance();
+		}
 	}
 }
 
diff --git a/Engine/Plugins/FX/Niagara/Source/Niagara/Private/NiagaraEmitterInstance.cpp b/Engine/Plugins/FX/Niagara/Source/Niagara/Private/NiagaraEmitterInstance.cpp
--- a/Engine/Plugins/FX/Niagara/Source/Niagara/Private/NiagaraEmitterInstance.cpp
+++ b/Engine/Plugins/FX/Niagara/Source/Niagara/Private/NiagaraEmitterInstance.cpp
@@ -739,6 +739,17 @@ void FNiagaraEmitterInstance::PreTick()
 		return;
 	}
 
+	if (TickCount == 0)
+	{
+		//On our very first frame we prime any previous params (for interpolation).
+		SpawnExecContext.PostTick();
+		UpdateExecContext.PostTick();
+		for (FNiagaraScriptExecutionContext& EventContext : EventExecContexts)
+		{
+			EventContext.PostTick();
+		}
+	}
+
 	checkSlow(Data.GetNumVariables() > 0);
 	checkSlow(CachedEmitter->SpawnScriptProps.Script);
 	checkSlow(CachedEmitter->UpdateScriptProps.Script);
diff --git a/Engine/Plugins/FX/Niagara/Source/Niagara/Private/NiagaraSystemSimulation.cpp b/Engine/Plugins/FX/Niagara/Source/Niagara/Private/NiagaraSystemSimulation.cpp
--- a/Engine/Plugins/FX/Niagara/Source/Niagara/Private/NiagaraSystemSimulation.cpp
+++ b/Engine/Plugins/FX/Niagara/Source/Niagara/Private/NiagaraSystemSimulation.cpp
@@ -32,6 +32,14 @@ static FAutoConsoleVariableRef CVarNiagaraDumpSystemData(
 	ECVF_Default
 );
 
+static int32 GbSystemUpdateOnSpawn = 1;
+static FAutoConsoleVariableRef CVarSystemUpdateOnSpawn(
+	TEXT("fx.SystemUpdateOnSpawn"),
+	GbSystemUpdateOnSpawn,
+	TEXT("If > 0, system simulations are given a small update after spawn. \n"),
+	ECVF_Default
+);
+
 //Pretick can no longer be run in parallel. Will likely remain this way.
 // static int32 GbParallelSystemPreTick = 0;
 // static FAutoConsoleVariableRef CVarNiagaraParallelSystemPreTick(
@@ -364,9 +372,6 @@ bool FNiagaraSystemSimulation::Tick(float DeltaSeconds)
 			SpawnInstanceParameterDataSet.Tick();
 			UpdateInstanceParameterDataSet.Tick();
 
-			DataSet.Tick();
-			DataSet.Allocate(NewNum);
-
 			//Setup the few real constants like delta time.
 			float InvDt = 1.0f / DeltaSeconds;
 			float GlobalSpawnCountScale = INiagaraModule::GetGlobalSpawnCountScale();
@@ -385,10 +390,52 @@ bool FNiagaraSystemSimulation::Tick(float DeltaSeconds)
 			DataSetExecInfos.Emplace(&DataSet, 0, false, true);
 		}
 
+		//TODO: JIRA - UE-60096 - Remove.
+		//We're having to allocate and spawn before update here so we have to do needless copies.			
+		//Ideally this should be compiled directly into the script similarly to interpolated particle spawning.
+		if (SpawnNum)
+		{
+			SCOPE_CYCLE_COUNTER(STAT_NiagaraSystemSim_Spawn);
+			DataSet.Allocate(NewNum, true);
+
+			DataSet.SetNumInstances(NewNum);
+			//SpawnInstanceParameterDataSet.SetNumInstances(NewNum);
+
+			//Run Spawn
+			SpawnExecContext.Tick(SoloSystemInstance);//We can't require a specific instance here as these are for all instances.
+			DataSetExecInfos.SetNum(1, false);
+			DataSetExecInfos[0].StartInstance = OrigNum;
+			DataSetExecInfos.Emplace(&SpawnInstanceParameterDataSet, OrigNum, false, false);
+			SpawnExecContext.Execute(SpawnNum, DataSetExecInfos);
+
+			if (GbDumpSystemData || System->bDumpDebugSystemInfo)
+			{
+				UE_LOG(LogNiagara, Log, TEXT("=== Spwaned %d Systems ==="), SpawnNum);
+				DataSet.Dump(true, OrigNum, SpawnNum);
+				SpawnInstanceParameterDataSet.Dump(false, OrigNum, SpawnNum);
+			}
+
+#if WITH_EDITORONLY_DATA
+			if (SoloSystemInstance && SoloSystemInstance->ShouldCaptureThisFrame())
+			{
+				FNiagaraScriptDebuggerInfo* DebugInfo = SoloSystemInstance->GetActiveCaptureWrite(NAME_None, ENiagaraScriptUsage::SystemSpawnScript, FGuid());
+				if (DebugInfo)
+				{
+					DataSet.Dump(DebugInfo->Frame, true, OrigNum, SpawnNum);
+					//DebugInfo->Frame.Dump(true, 0, SpawnNum);
+					DebugInfo->Parameters = SpawnExecContext.Parameters;
+				}
+			}
+#endif
+		}
+
+		DataSet.Tick();
+		DataSet.Allocate(NewNum);
+
 		{
 			SCOPE_CYCLE_COUNTER(STAT_NiagaraSystemSim_Update);
-			DataSet.SetNumInstances(OrigNum);
-			UpdateInstanceParameterDataSet.SetNumInstances(OrigNum);
+			DataSet.SetNumInstances(NewNum);
+			//UpdateInstanceParameterDataSet.SetNumInstances(OrigNum);
 
 			//Run update.
 			UpdateExecContext.Tick(SystemInstances[0]);
@@ -410,13 +457,37 @@ bool FNiagaraSystemSimulation::Tick(float DeltaSeconds)
 				DataSet.Dump(true, 0, OrigNum);
 			}
 
+ 			//Also run the update script on the newly spawned systems too.
+			//TODO: JIRA - UE-60096 - Remove.
+			//Ideally this should be compiled directly into the script similarly to interpolated particle spawning.
+			if (SpawnNum && GbSystemUpdateOnSpawn)
+			{
+				DataSet.SetNumInstances(NewNum);
+
+				//Run update.
+				UpdateExecContext.Tick(SystemInstances[0]);
+				DataSetExecInfos[0].StartInstance = OrigNum;
+				DataSetExecInfos[1].StartInstance = OrigNum;
+
+				UpdateExecContext.Parameters.SetParameterValue(0.0001f, SYS_PARAM_ENGINE_DELTA_TIME);
+				UpdateExecContext.Parameters.SetParameterValue(10000.0f, SYS_PARAM_ENGINE_INV_DELTA_TIME);
+
+				UpdateExecContext.Execute(SpawnNum, DataSetExecInfos);
+
+				if (GbDumpSystemData || System->bDumpDebugSystemInfo)
+				{
+					UE_LOG(LogNiagara, Log, TEXT("=== Spawn Updated %d Systems ==="), SpawnNum);
+					DataSet.Dump(true, OrigNum, SpawnNum);
+				}
+			}
+
 #if WITH_EDITORONLY_DATA
 			if (SoloSystemInstance && SoloSystemInstance->ShouldCaptureThisFrame())
 			{
 				FNiagaraScriptDebuggerInfo* DebugInfo = SoloSystemInstance->GetActiveCaptureWrite(NAME_None, ENiagaraScriptUsage::SystemUpdateScript, FGuid());
 				if (DebugInfo)
 				{
-					DataSet.Dump(DebugInfo->Frame, true, 0, OrigNum);
+					DataSet.Dump(DebugInfo->Frame, true, 0, NewNum);
 					//DebugInfo->Frame.Dump(true, 0, OrigNum);
 					DebugInfo->Parameters = UpdateExecContext.Parameters;
 				}
@@ -434,40 +505,6 @@ bool FNiagaraSystemSimulation::Tick(float DeltaSeconds)
 			}
 		}
 
-		if(SpawnNum)
-		{
-			SCOPE_CYCLE_COUNTER(STAT_NiagaraSystemSim_Spawn);
-			DataSet.SetNumInstances(NewNum);
-			SpawnInstanceParameterDataSet.SetNumInstances(NewNum);
-
-			//Run Spawn
-			SpawnExecContext.Tick(SoloSystemInstance);//We can't require a specific instance here as these are for all instances.
-			DataSetExecInfos.SetNum(1, false);
-			DataSetExecInfos[0].StartInstance = OrigNum;
-			DataSetExecInfos.Emplace(&SpawnInstanceParameterDataSet, OrigNum, false, false);
-			SpawnExecContext.Execute(SpawnNum, DataSetExecInfos);
-
-			if (GbDumpSystemData || System->bDumpDebugSystemInfo)
-			{
-				UE_LOG(LogNiagara, Log, TEXT("=== Spwaned %d Systems ==="), SpawnNum);
-				DataSet.Dump(true, OrigNum, SpawnNum);
-				SpawnInstanceParameterDataSet.Dump(false, OrigNum, SpawnNum);
-			}
-
-#if WITH_EDITORONLY_DATA
-			if (SoloSystemInstance && SoloSystemInstance->ShouldCaptureThisFrame())
-			{
-				FNiagaraScriptDebuggerInfo* DebugInfo = SoloSystemInstance->GetActiveCaptureWrite(NAME_None, ENiagaraScriptUsage::SystemSpawnScript, FGuid());
-				if (DebugInfo)
-				{
-					DataSet.Dump(DebugInfo->Frame, true, OrigNum, SpawnNum);
-					//DebugInfo->Frame.Dump(true, 0, SpawnNum);
-					DebugInfo->Parameters = SpawnExecContext.Parameters;
-				}
-			}
-#endif
-		}
-
 		{
 			SCOPE_CYCLE_COUNTER(STAT_NiagaraSystemSim_TransferParameters);
 			int32 SystemIndex = 0;
diff --git a/Engine/Plugins/FX/Niagara/Source/NiagaraCore/Private/NiagaraCustomVersion.cpp b/Engine/Plugins/FX/Niagara/Source/NiagaraCore/Private/NiagaraCustomVersion.cpp
--- a/Engine/Plugins/FX/Niagara/Source/NiagaraCore/Private/NiagaraCustomVersion.cpp
+++ b/Engine/Plugins/FX/Niagara/Source/NiagaraCore/Private/NiagaraCustomVersion.cpp
@@ -8,4 +8,4 @@ const FGuid FNiagaraCustomVersion::GUID(0xFCF57AFA, 0x50764283, 0xB9A9E658, 0xFF
 // Register the custom version with core
 FCustomVersionRegistration GRegisterNiagaraCustomVersion(FNiagaraCustomVersion::GUID, FNiagaraCustomVersion::LatestVersion, TEXT("NiagaraVer"));
 
-const FGuid FNiagaraCustomVersion::LatestScriptCompileVersion(0xDF96299F, 0x970043FE, 0x9CE8ADF0, 0xD2D4E1AA);
\ No newline at end of file
+const FGuid FNiagaraCustomVersion::LatestScriptCompileVersion(0xEAA7D8F4, 0x8BE04550, 0xAF0EC099, 0x5A99A2F4);
\ No newline at end of file
diff --git a/Engine/Plugins/FX/Niagara/Source/NiagaraEditor/Private/EdGraphSchema_Niagara.cpp b/Engine/Plugins/FX/Niagara/Source/NiagaraEditor/Private/EdGraphSchema_Niagara.cpp
--- a/Engine/Plugins/FX/Niagara/Source/NiagaraEditor/Private/EdGraphSchema_Niagara.cpp
+++ b/Engine/Plugins/FX/Niagara/Source/NiagaraEditor/Private/EdGraphSchema_Niagara.cpp
@@ -182,6 +182,49 @@ TSharedPtr<FNiagaraSchemaAction_NewNode> AddNewNodeAction(TArray<TSharedPtr<FNia
 	return NewAction;
 }
 
+
+
+bool IsSystemGraph(const UNiagaraGraph* NiagaraGraph)
+{
+	TArray<UNiagaraNodeEmitter*> Emitters;
+	NiagaraGraph->GetNodesOfClass<UNiagaraNodeEmitter>(Emitters);
+	bool bSystemGraph = Emitters.Num() != 0 || NiagaraGraph->FindOutputNode(ENiagaraScriptUsage::SystemSpawnScript) != nullptr || NiagaraGraph->FindOutputNode(ENiagaraScriptUsage::SystemUpdateScript) != nullptr;
+	return bSystemGraph;
+}
+
+bool IsParticleGraph(const UNiagaraGraph* NiagaraGraph)
+{
+	bool bParticleGraph = NiagaraGraph->FindOutputNode(ENiagaraScriptUsage::ParticleSpawnScriptInterpolated) != nullptr || NiagaraGraph->FindOutputNode(ENiagaraScriptUsage::ParticleSpawnScript) != nullptr || NiagaraGraph->FindOutputNode(ENiagaraScriptUsage::ParticleUpdateScript) != nullptr;
+	return bParticleGraph;
+}
+
+bool IsModuleGraph(const UNiagaraGraph* NiagaraGraph)
+{
+	bool bModuleGraph = NiagaraGraph->FindOutputNode(ENiagaraScriptUsage::Module) != nullptr;
+	return bModuleGraph;
+}
+
+bool IsDynamicInputGraph(const UNiagaraGraph* NiagaraGraph)
+{
+	bool bDynamicInputGraph = NiagaraGraph->FindOutputNode(ENiagaraScriptUsage::DynamicInput) != nullptr;
+	return bDynamicInputGraph;
+}
+
+
+bool IsFunctionGraph(const UNiagaraGraph* NiagaraGraph)
+{
+	bool bFunctionGraph = NiagaraGraph->FindOutputNode(ENiagaraScriptUsage::Function) != nullptr;
+	return bFunctionGraph;
+}
+
+
+bool IsParticleUpdateGraph(const UNiagaraGraph* NiagaraGraph)
+{
+	bool bUpdateGraph = NiagaraGraph->FindOutputNode(ENiagaraScriptUsage::ParticleUpdateScript) != nullptr;
+	return bUpdateGraph;
+}
+
+
 const UNiagaraGraph* GetAlternateGraph(const UNiagaraGraph* NiagaraGraph)
 {
 	UNiagaraScriptSource* ScriptSource = Cast<UNiagaraScriptSource>(NiagaraGraph->GetOuter());
@@ -299,16 +342,13 @@ TArray<TSharedPtr<FNiagaraSchemaAction_NewNode> > UEdGraphSchema_Niagara::GetGra
 	TArray<TSharedPtr<FNiagaraSchemaAction_NewNode> > NewActions;
 
 	const UNiagaraGraph* NiagaraGraph = CastChecked<UNiagaraGraph>(CurrentGraph);
-	TArray<UNiagaraNodeEmitter*> Emitters;
-	NiagaraGraph->GetNodesOfClass<UNiagaraNodeEmitter>(Emitters);
-	bool bSystemGraph = Emitters.Num() != 0 || NiagaraGraph->FindOutputNode(ENiagaraScriptUsage::SystemSpawnScript) != nullptr || NiagaraGraph->FindOutputNode(ENiagaraScriptUsage::SystemUpdateScript) != nullptr;
-	bool bModuleGraph = NiagaraGraph->FindOutputNode(ENiagaraScriptUsage::Module) != nullptr;
-	bool bDynamicInputGraph = NiagaraGraph->FindOutputNode(ENiagaraScriptUsage::DynamicInput) != nullptr;
-	bool bFunctionGraph = NiagaraGraph->FindOutputNode(ENiagaraScriptUsage::Function) != nullptr;
-	bool bUpdateGraph = NiagaraGraph->FindOutputNode(ENiagaraScriptUsage::ParticleUpdateScript) != nullptr;
-
 	
-
+	bool bSystemGraph = IsSystemGraph(NiagaraGraph);
+	bool bModuleGraph = IsModuleGraph(NiagaraGraph);
+	bool bDynamicInputGraph = IsDynamicInputGraph(NiagaraGraph);
+	bool bFunctionGraph = IsFunctionGraph(NiagaraGraph);
+	bool bParticleUpdateGraph = IsParticleUpdateGraph(NiagaraGraph);
+	
 	if (GbAllowAllNiagaraNodesInEmitterGraphs || bModuleGraph || bFunctionGraph || bSystemGraph)
 	{
 		const TArray<FNiagaraOpInfo>& OpInfos = FNiagaraOpInfo::GetOpInfoArray();
@@ -356,7 +396,7 @@ TArray<TSharedPtr<FNiagaraSchemaAction_NewNode> > UEdGraphSchema_Niagara::GetGra
 	TArray<FAssetData> ScriptAssets;
 	AssetRegistryModule.Get().GetAssetsByClass(UNiagaraScript::StaticClass()->GetFName(), ScriptAssets);
 	UEnum* NiagaraScriptUsageEnum = FindObjectChecked<UEnum>(ANY_PACKAGE, TEXT("ENiagaraScriptUsage"), true);
-	if (GbAllowAllNiagaraNodesInEmitterGraphs || bModuleGraph || bFunctionGraph || bSystemGraph)
+	if (GbAllowAllNiagaraNodesInEmitterGraphs || bModuleGraph || bFunctionGraph || bDynamicInputGraph)
 	{
 		for (const FAssetData& ScriptAsset : ScriptAssets)
 		{
@@ -387,7 +427,7 @@ TArray<TSharedPtr<FNiagaraSchemaAction_NewNode> > UEdGraphSchema_Niagara::GetGra
 			int32 BitfieldValue = FCString::Atoi(*BitfieldTagValue);
 
 			ENiagaraScriptUsage TargetUsage = ENiagaraScriptUsage::Module;
-			if (bUpdateGraph)
+			if (bParticleUpdateGraph)
 			{
 				TargetUsage = ENiagaraScriptUsage::ParticleUpdateScript;
 			}
@@ -465,7 +505,7 @@ TArray<TSharedPtr<FNiagaraSchemaAction_NewNode> > UEdGraphSchema_Niagara::GetGra
 	}
 	
 	TArray<ENiagaraScriptUsage> UsageTypesToAdd;
-	if (bUpdateGraph)
+	if (bParticleUpdateGraph)
 	{
 		UsageTypesToAdd.Add(ENiagaraScriptUsage::ParticleEventScript);
 		UsageTypesToAdd.Add(ENiagaraScriptUsage::EmitterSpawnScript);
@@ -717,20 +757,23 @@ TArray<TSharedPtr<FNiagaraSchemaAction_NewNode> > UEdGraphSchema_Niagara::GetGra
 		TArray<UNiagaraNodeInput*> InputNodes;
 		NiagaraGraph->GetNodesOfClass(InputNodes);
 
-		//Emitter constants managed by the system.
-		const TArray<FNiagaraVariable>& SystemConstants = FNiagaraConstants::GetEngineConstants();
-		for (const FNiagaraVariable& SysConst : SystemConstants)
+		if (bFunctionGraph)
 		{
-			FFormatNamedArguments Args;
-			Args.Add(TEXT("Constant"), FText::FromName(SysConst.GetName()));
-			const FText MenuDesc = FText::Format(LOCTEXT("GetSystemConstant", "Get {Constant}"), Args);
+			//Emitter constants managed by the system.
+			const TArray<FNiagaraVariable>& SystemConstants = FNiagaraConstants::GetEngineConstants();
+			for (const FNiagaraVariable& SysConst : SystemConstants)
+			{
+				FFormatNamedArguments Args;
+				Args.Add(TEXT("Constant"), FText::FromName(SysConst.GetName()));
+				const FText MenuDesc = FText::Format(LOCTEXT("GetSystemConstant", "Get {Constant}"), Args);
 
-			TSharedPtr<FNiagaraSchemaAction_NewNode> GetConstAction = AddNewNodeAction(NewActions, LOCTEXT("System Parameters Menu Title", "System Parameters"), MenuDesc, SysConst.GetName(), FText::GetEmpty());
+				TSharedPtr<FNiagaraSchemaAction_NewNode> GetConstAction = AddNewNodeAction(NewActions, LOCTEXT("System Parameters Menu Title", "System Parameters"), MenuDesc, SysConst.GetName(), FText::GetEmpty());
 
-			UNiagaraNodeInput* InputNode = NewObject<UNiagaraNodeInput>(OwnerOfTemporaries);
-			InputNode->Usage = ENiagaraInputNodeUsage::SystemConstant;
-			InputNode->Input = SysConst;
-			GetConstAction->NodeTemplate = InputNode;
+				UNiagaraNodeInput* InputNode = NewObject<UNiagaraNodeInput>(OwnerOfTemporaries);
+				InputNode->Usage = ENiagaraInputNodeUsage::SystemConstant;
+				InputNode->Input = SysConst;
+				GetConstAction->NodeTemplate = InputNode;
+			}
 		}
 
 		//Emitter constants managed by the Translator.
@@ -777,6 +820,16 @@ TArray<TSharedPtr<FNiagaraSchemaAction_NewNode> > UEdGraphSchema_Niagara::GetGra
 					else
 					{
 						MenuCat = LOCTEXT("AddParameterCat", "Add Parameter");
+
+						// If you are in dynamic inputs or modules, we only allow free-range variables for 
+						// data interfaces and parameter maps.
+						if (bDynamicInputGraph || bModuleGraph)
+						{
+							if (Type != FNiagaraTypeDefinition::GetParameterMapDef())
+							{
+								continue;
+							}
+						}
 					}
 						
 					const FText MenuDesc = FText::Format(MenuDescFmt, Type.GetStruct()->GetDisplayNameText());
@@ -787,24 +840,27 @@ TArray<TSharedPtr<FNiagaraSchemaAction_NewNode> > UEdGraphSchema_Niagara::GetGra
 				}
 
 				// TODO sckime please remove this..
-				for (FNiagaraTypeDefinition Type : RegisteredTypes)
+				if (bSystemGraph || IsParticleGraph(NiagaraGraph))
 				{
-					FText MenuCat;
-					if (const UClass* Class = Type.GetClass())
-					{
-						continue;
-					}
-					else
+					for (FNiagaraTypeDefinition Type : RegisteredTypes)
 					{
-						MenuCat = LOCTEXT("AddRIParameterCat", "Add Rapid Iteration Param");
-					}
+						FText MenuCat;
+						if (const UClass* Class = Type.GetClass())
+						{
+							continue;
+						}
+						else
+						{
+							MenuCat = LOCTEXT("AddRIParameterCat", "Add Rapid Iteration Param");
+						}
 
-					const FText MenuDesc = FText::Format(MenuDescFmt, Type.GetStruct()->GetDisplayNameText());
-					TSharedPtr<FNiagaraSchemaAction_NewNode> InputAction = AddNewNodeAction(NewActions, MenuCat, MenuDesc, *MenuDesc.ToString(), FText::GetEmpty());
-					UNiagaraNodeInput* InputNode = NewObject<UNiagaraNodeInput>(OwnerOfTemporaries);
-					FNiagaraEditorUtilities::InitializeParameterInputNode(*InputNode, Type, NiagaraGraph);
-					InputNode->Usage = ENiagaraInputNodeUsage::RapidIterationParameter;
-					InputAction->NodeTemplate = InputNode;
+						const FText MenuDesc = FText::Format(MenuDescFmt, Type.GetStruct()->GetDisplayNameText());
+						TSharedPtr<FNiagaraSchemaAction_NewNode> InputAction = AddNewNodeAction(NewActions, MenuCat, MenuDesc, *MenuDesc.ToString(), FText::GetEmpty());
+						UNiagaraNodeInput* InputNode = NewObject<UNiagaraNodeInput>(OwnerOfTemporaries);
+						FNiagaraEditorUtilities::InitializeParameterInputNode(*InputNode, Type, NiagaraGraph);
+						InputNode->Usage = ENiagaraInputNodeUsage::RapidIterationParameter;
+						InputAction->NodeTemplate = InputNode;
+					}
 				}
 				
 				if (PinType != FNiagaraTypeDefinition::GetGenericNumericDef())
@@ -1498,6 +1554,16 @@ void UEdGraphSchema_Niagara::ToggleNodeEnabledState(UNiagaraNode* InNode) const
 	}
 }
 
+bool UEdGraphSchema_Niagara::CanPromoteSinglePinToParameter(const UEdGraphPin* SourcePin) 
+{
+	const UNiagaraGraph* NiagaraGraph = Cast<UNiagaraGraph>(SourcePin->GetOwningNode()->GetGraph());
+	if (IsFunctionGraph(NiagaraGraph))
+	{
+		return true;
+	}
+	return false;
+}
+
 void UEdGraphSchema_Niagara::PromoteSinglePinToParameter(UEdGraphPin* SourcePin)
 {
 	if (SourcePin)
@@ -1560,7 +1626,8 @@ void UEdGraphSchema_Niagara::GetContextMenuActions(const UEdGraph* CurrentGraph,
 			if (InGraphPin->Direction == EEdGraphPinDirection::EGPD_Input)
 			{
 				MenuBuilder->AddMenuEntry(LOCTEXT("PromoteToParameter", "Promote to Parameter"), LOCTEXT("PromoteToParameterTooltip", "Create a parameter argument and connect this pin to that parameter."), FSlateIcon(),
-					FUIAction(FExecuteAction::CreateUObject((UEdGraphSchema_Niagara*const)this, &UEdGraphSchema_Niagara::PromoteSinglePinToParameter, const_cast<UEdGraphPin*>(InGraphPin))));
+					FUIAction(FExecuteAction::CreateUObject((UEdGraphSchema_Niagara*const)this, &UEdGraphSchema_Niagara::PromoteSinglePinToParameter, const_cast<UEdGraphPin*>(InGraphPin)),
+						FCanExecuteAction::CreateStatic(&UEdGraphSchema_Niagara::CanPromoteSinglePinToParameter, InGraphPin)));
 				if (InGraphPin->LinkedTo.Num() == 0 && InGraphPin->bDefaultValueIsIgnored == false)
 				{
 					MenuBuilder->AddMenuEntry(
diff --git a/Engine/Plugins/FX/Niagara/Source/NiagaraEditor/Private/NiagaraEditorModule.cpp b/Engine/Plugins/FX/Niagara/Source/NiagaraEditor/Private/NiagaraEditorModule.cpp
--- a/Engine/Plugins/FX/Niagara/Source/NiagaraEditor/Private/NiagaraEditorModule.cpp
+++ b/Engine/Plugins/FX/Niagara/Source/NiagaraEditor/Private/NiagaraEditorModule.cpp
@@ -87,6 +87,7 @@
 #include "Customizations/NiagaraTypeCustomizations.h"
 #include "Customizations/NiagaraEventScriptPropertiesCustomization.h"
 #include "HAL/IConsoleManager.h"
+#include "NiagaraHlslTranslator.h"
 
 
 IMPLEMENT_MODULE( FNiagaraEditorModule, NiagaraEditor );
@@ -281,6 +282,7 @@ class FNiagaraSystemColorParameterTrackEditor : public FNiagaraSystemParameterTr
 
 void FNiagaraEditorModule::StartupModule()
 {
+	FHlslNiagaraTranslator::Init();
 	MenuExtensibilityManager = MakeShareable(new FExtensibilityManager);
 	ToolBarExtensibilityManager = MakeShareable(new FExtensibilityManager);
 
diff --git a/Engine/Plugins/FX/Niagara/Source/NiagaraEditor/Private/NiagaraHlslTranslator.cpp b/Engine/Plugins/FX/Niagara/Source/NiagaraEditor/Private/NiagaraHlslTranslator.cpp
--- a/Engine/Plugins/FX/Niagara/Source/NiagaraEditor/Private/NiagaraHlslTranslator.cpp
+++ b/Engine/Plugins/FX/Niagara/Source/NiagaraEditor/Private/NiagaraHlslTranslator.cpp
@@ -564,10 +564,11 @@ FString FHlslNiagaraTranslator::BuildParameterMapHlslDefinitions(TArray<FNiagara
 		}
 
 		TArray<FString> StructNameArray;
-		int32 NumFound = Variable.GetName().ToString().ParseIntoArray(StructNameArray, TEXT("."));
+		FString SanitizedVarName = GetSanitizedSymbolName(Variable.GetName().ToString());
+		int32 NumFound = SanitizedVarName.ParseIntoArray(StructNameArray, TEXT("."));
 		if (NumFound == 1) // Meaning no split above
 		{
-			Error(FText::Format(LOCTEXT("OnlyOneNamespaceEntry", "Only one namespace entry found for: {0}"), FText::FromName(Variable.GetName())), nullptr, nullptr);
+			Error(FText::Format(LOCTEXT("OnlyOneNamespaceEntry", "Only one namespace entry found for: {0}"), FText::FromString(SanitizedVarName)), nullptr, nullptr);
 		}
 		else if (NumFound > 1)
 		{
@@ -1173,7 +1174,7 @@ const FNiagaraTranslateResults &FHlslNiagaraTranslator::Translate(const FNiagara
 		{
 			if (FNiagaraParameterMapHistory::IsAttribute(Var))
 			{
-				FNiagaraVariable BasicAttribVar = FNiagaraParameterMapHistory::ResolveAsBasicAttribute(Var);
+				FNiagaraVariable BasicAttribVar = FNiagaraParameterMapHistory::ResolveAsBasicAttribute(Var, false);
 				CompilationOutput.ScriptData.Attributes.AddUnique(BasicAttribVar);
 			}
 			else
@@ -2063,15 +2064,432 @@ void FHlslNiagaraTranslator::DecomposeVariableAccess(UStruct* Struct, bool bRead
 	}
 };
 
+TMap<FString, FString> FHlslNiagaraTranslator::ReplacementsForInvalid;
+
+void FHlslNiagaraTranslator::Init()
+{
+	if (ReplacementsForInvalid.Num() == 0)
+	{
+		// hlsl keywords
+		ReplacementsForInvalid.Add(TEXT("AppendStructuredBuffer"), TEXT("APPENDSTRUCTUREDBUFFER_VAR"));
+		ReplacementsForInvalid.Add(TEXT("asm"), TEXT("ASSEMBLY_VAR"));
+		ReplacementsForInvalid.Add(TEXT("asm_fragment"), TEXT("ASSEMBLY_FRAGMENT_VAR"));
+		ReplacementsForInvalid.Add(TEXT("BlendState"), TEXT("BLEND_STATE_VAR"));
+		ReplacementsForInvalid.Add(TEXT("bool"), TEXT("BOOL_VAR"));
+		ReplacementsForInvalid.Add(TEXT("bool1"), TEXT("BOOL1_VAR"));
+		ReplacementsForInvalid.Add(TEXT("bool2"), TEXT("BOOL2_VAR"));
+		ReplacementsForInvalid.Add(TEXT("bool3"), TEXT("BOOL3_VAR"));
+		ReplacementsForInvalid.Add(TEXT("bool4"), TEXT("BOOL4_VAR"));
+		ReplacementsForInvalid.Add(TEXT("break"), TEXT("BREAK_VAR"));
+		ReplacementsForInvalid.Add(TEXT("Buffer"), TEXT("BUFFER_VAR"));
+		ReplacementsForInvalid.Add(TEXT("ByteAddressBuffer"), TEXT("BYTE_ADDRESS_BUFFER_VAR"));
+		ReplacementsForInvalid.Add(TEXT("case"), TEXT("CASE_VAR"));
+		ReplacementsForInvalid.Add(TEXT("cbuffer"), TEXT("CBUFFER_VAR"));
+		ReplacementsForInvalid.Add(TEXT("centroid"), TEXT("CENTROID_VAR"));
+		ReplacementsForInvalid.Add(TEXT("class"), TEXT("CLASS_VAR"));
+		ReplacementsForInvalid.Add(TEXT("column_major"), TEXT("COLUMN_MAJOR_VAR"));
+		ReplacementsForInvalid.Add(TEXT("compile"), TEXT("COMPILE_VAR"));
+		ReplacementsForInvalid.Add(TEXT("compile_fragment"), TEXT("COMPILE_FRAGMENT_VAR"));
+		ReplacementsForInvalid.Add(TEXT("CompileShader"), TEXT("COMPILE_SHADER_VAR"));
+		ReplacementsForInvalid.Add(TEXT("const"), TEXT("CONST_VAR"));
+		ReplacementsForInvalid.Add(TEXT("continue"), TEXT("CONTINUE_VAR"));
+		ReplacementsForInvalid.Add(TEXT("ComputeShader"), TEXT("COMPUTE_SHADER_VAR"));
+		ReplacementsForInvalid.Add(TEXT("ConsumeStructuredBuffer"), TEXT("CONSUME_STRUCTURED_BUFFER"));
+		ReplacementsForInvalid.Add(TEXT("default"), TEXT("DEFAULT_VAR"));
+		ReplacementsForInvalid.Add(TEXT("DepthStencilState"), TEXT("DEPTHSTENCILSTATE_VAR"));
+		ReplacementsForInvalid.Add(TEXT("DepthStencilView"), TEXT("DEPTHSTENCILVIEW_VAR"));
+		ReplacementsForInvalid.Add(TEXT("discard"), TEXT("DISCARD_VAR"));
+		ReplacementsForInvalid.Add(TEXT("do"), TEXT("DO_VAR"));
+		ReplacementsForInvalid.Add(TEXT("double"), TEXT("DOUBLE_VAR"));
+		ReplacementsForInvalid.Add(TEXT("double1"), TEXT("DOUBLE1_VAR"));
+		ReplacementsForInvalid.Add(TEXT("double2"), TEXT("DOUBLE2_VAR"));
+		ReplacementsForInvalid.Add(TEXT("double3"), TEXT("DOUBLE3_VAR"));
+		ReplacementsForInvalid.Add(TEXT("double4"), TEXT("DOUBLE4_VAR"));
+		ReplacementsForInvalid.Add(TEXT("DomainShader"), TEXT("DOMAINSHADER_VAR"));
+		ReplacementsForInvalid.Add(TEXT("dword"), TEXT("DWORD_VAR"));
+		ReplacementsForInvalid.Add(TEXT("else"), TEXT("ELSE_VAR"));
+		ReplacementsForInvalid.Add(TEXT("export"), TEXT("EXPORT_VAR"));
+		ReplacementsForInvalid.Add(TEXT("extern"), TEXT("EXTERN_VAR"));
+		ReplacementsForInvalid.Add(TEXT("false"), TEXT("FALSE_VAR"));
+		ReplacementsForInvalid.Add(TEXT("float"), TEXT("FLOAT_VAR"));
+		ReplacementsForInvalid.Add(TEXT("float1"), TEXT("FLOAT1_VAR"));
+		ReplacementsForInvalid.Add(TEXT("float2"), TEXT("FLOAT2_VAR"));
+		ReplacementsForInvalid.Add(TEXT("float3"), TEXT("FLOAT3_VAR"));
+		ReplacementsForInvalid.Add(TEXT("float4"), TEXT("FLOAT4_VAR"));
+		ReplacementsForInvalid.Add(TEXT("for"), TEXT("FOR_VAR"));
+		ReplacementsForInvalid.Add(TEXT("fxgroup"), TEXT("FXGROUP_VAR"));
+		ReplacementsForInvalid.Add(TEXT("GeometryShader"), TEXT("GEOMETRYSHADER_VAR"));
+		ReplacementsForInvalid.Add(TEXT("groupshared"), TEXT("GROUPSHARED_VAR"));
+		ReplacementsForInvalid.Add(TEXT("half"), TEXT("HALF_VAR"));
+		ReplacementsForInvalid.Add(TEXT("half1"), TEXT("HALF1_VAR"));
+		ReplacementsForInvalid.Add(TEXT("half2"), TEXT("HALF2_VAR"));
+		ReplacementsForInvalid.Add(TEXT("half3"), TEXT("HALF3_VAR"));
+		ReplacementsForInvalid.Add(TEXT("half4"), TEXT("HALF4_VAR"));
+		ReplacementsForInvalid.Add(TEXT("Hullshader"), TEXT("HULLSHADER_VAR"));
+		ReplacementsForInvalid.Add(TEXT("if"), TEXT("IF_VAR"));
+		ReplacementsForInvalid.Add(TEXT("in"), TEXT("IN_VAR"));
+		ReplacementsForInvalid.Add(TEXT("inline"), TEXT("INLINE_VAR"));
+		ReplacementsForInvalid.Add(TEXT("inout"), TEXT("INOUT_VAR"));
+		ReplacementsForInvalid.Add(TEXT("InputPatch"), TEXT("INPUTPATCH_VAR"));
+		ReplacementsForInvalid.Add(TEXT("int"), TEXT("INT_VAR"));
+		ReplacementsForInvalid.Add(TEXT("int1"), TEXT("INT1_VAR"));
+		ReplacementsForInvalid.Add(TEXT("int2"), TEXT("INT2_VAR"));
+		ReplacementsForInvalid.Add(TEXT("int3"), TEXT("INT3_VAR"));
+		ReplacementsForInvalid.Add(TEXT("int4"), TEXT("INT4_VAR"));
+		ReplacementsForInvalid.Add(TEXT("interface"), TEXT("INTERFACE_VAR"));
+		ReplacementsForInvalid.Add(TEXT("line"), TEXT("LINE_VAR"));
+		ReplacementsForInvalid.Add(TEXT("lineadj"), TEXT("LINEADJ_VAR"));
+		ReplacementsForInvalid.Add(TEXT("linear"), TEXT("LINEAR_VAR"));
+		ReplacementsForInvalid.Add(TEXT("LineStream"), TEXT("LINESTREAM_VAR"));
+		ReplacementsForInvalid.Add(TEXT("matrix"), TEXT("MATRIX_VAR"));
+		ReplacementsForInvalid.Add(TEXT("min16float"), TEXT("MIN16FLOAT_VAR"));
+		ReplacementsForInvalid.Add(TEXT("min10float"), TEXT("MIN10FLOAT_VAR"));
+		ReplacementsForInvalid.Add(TEXT("min16int"), TEXT("MIN16INT_VAR"));
+		ReplacementsForInvalid.Add(TEXT("min12int"), TEXT("MIN12INT_VAR"));
+		ReplacementsForInvalid.Add(TEXT("min16uint"), TEXT("MIN16UINT_VAR"));
+		ReplacementsForInvalid.Add(TEXT("namespace"), TEXT("NAMESPACE_VAR"));
+		ReplacementsForInvalid.Add(TEXT("nointerpolation"), TEXT("NOINTERPOLATION_VAR"));
+		ReplacementsForInvalid.Add(TEXT("noperspective"), TEXT("NOPERSPECTIVE_VAR"));
+		ReplacementsForInvalid.Add(TEXT("NULL"), TEXT("NULL_VAR"));
+		ReplacementsForInvalid.Add(TEXT("out"), TEXT("OUT_VAR"));
+		ReplacementsForInvalid.Add(TEXT("OutputPatch"), TEXT("OUTPUTPATCH_VAR"));
+		ReplacementsForInvalid.Add(TEXT("packoffset"), TEXT("PACKOFFSET,_VAR"));
+		ReplacementsForInvalid.Add(TEXT("pass"), TEXT("PASS,_VAR"));
+		ReplacementsForInvalid.Add(TEXT("pixelfragment"), TEXT("PIXELFRAGMENT_VAR"));
+		ReplacementsForInvalid.Add(TEXT("PixelShader"), TEXT("PIXELSHADER_VAR"));
+		ReplacementsForInvalid.Add(TEXT("point"), TEXT("POINT_VAR"));
+		ReplacementsForInvalid.Add(TEXT("PointStream"), TEXT("POINTSTREAM_VAR"));
+		ReplacementsForInvalid.Add(TEXT("precise"), TEXT("PRECISE_VAR"));
+		ReplacementsForInvalid.Add(TEXT("RasterizerState"), TEXT("RASTERIZERSTATE_VAR"));
+		ReplacementsForInvalid.Add(TEXT("RenderTargetView"), TEXT("RENDERTARGETVIEW_VAR"));
+		ReplacementsForInvalid.Add(TEXT("return"), TEXT("RETURN_VAR"));
+		ReplacementsForInvalid.Add(TEXT("register"), TEXT("REGISTER_VAR"));
+		ReplacementsForInvalid.Add(TEXT("row_major"), TEXT("ROW_MAJOR_VAR"));
+		ReplacementsForInvalid.Add(TEXT("RWBuffer"), TEXT("RWBUFFER_VAR"));
+		ReplacementsForInvalid.Add(TEXT("RWByteAddressBuffer"), TEXT("RWBYTEADDRESSBUFFER_VAR"));
+		ReplacementsForInvalid.Add(TEXT("RWStructuredBuffer"), TEXT("RWSTRUCTUREDBUFFER_VAR"));
+		ReplacementsForInvalid.Add(TEXT("RWTexture1D"), TEXT("RWTEXTURE1D_VAR"));
+		ReplacementsForInvalid.Add(TEXT("RWTexture1DArray"), TEXT("RWTEXTURE1DARRAY_VAR"));
+		ReplacementsForInvalid.Add(TEXT("RWTexture2D"), TEXT("RWTEXTURE2D_VAR"));
+		ReplacementsForInvalid.Add(TEXT("RWTexture2DArray"), TEXT("RWTEXTURE2DARRAY_VAR"));
+		ReplacementsForInvalid.Add(TEXT("RWTexture3D"), TEXT("RWTEXTURE3D_VAR"));
+		ReplacementsForInvalid.Add(TEXT("sample"), TEXT("SAMPLE_VAR"));
+		ReplacementsForInvalid.Add(TEXT("sampler"), TEXT("SAMPLER_VAR"));
+		ReplacementsForInvalid.Add(TEXT("SamplerState"), TEXT("SAMPLERSTATE_VAR"));
+		ReplacementsForInvalid.Add(TEXT("SamplerComparisonState"), TEXT("SAMPLERCOMPARISONSTATE_VAR"));
+		ReplacementsForInvalid.Add(TEXT("shared"), TEXT("SHARED_VAR"));
+		ReplacementsForInvalid.Add(TEXT("snorm"), TEXT("SNORM_VAR"));
+		ReplacementsForInvalid.Add(TEXT("stateblock"), TEXT("STATEBLOCK_VAR"));
+		ReplacementsForInvalid.Add(TEXT("stateblock_state"), TEXT("STATEBLOCK_STATE_VAR"));
+		ReplacementsForInvalid.Add(TEXT("static"), TEXT("STATIC_VAR"));
+		ReplacementsForInvalid.Add(TEXT("string"), TEXT("STRING_VAR"));
+		ReplacementsForInvalid.Add(TEXT("struct"), TEXT("STRUCT_VAR"));
+		ReplacementsForInvalid.Add(TEXT("switch"), TEXT("SWITCH_VAR"));
+		ReplacementsForInvalid.Add(TEXT("StructuredBuffer"), TEXT("STRUCTUREDBUFFER_VAR"));
+		ReplacementsForInvalid.Add(TEXT("tbuffer"), TEXT("TBUFFER_VAR"));
+		ReplacementsForInvalid.Add(TEXT("technique"), TEXT("TECHNIQUE_VAR"));
+		ReplacementsForInvalid.Add(TEXT("technique10"), TEXT("TECHNIQUE10_VAR"));
+		ReplacementsForInvalid.Add(TEXT("technique11"), TEXT("TECHNIQUE11_VAR"));
+		ReplacementsForInvalid.Add(TEXT("texture"), TEXT("TEXTURE_VAR"));
+		ReplacementsForInvalid.Add(TEXT("Texture1D"), TEXT("TEXTURE1D_VAR"));
+		ReplacementsForInvalid.Add(TEXT("Texture1DArray"), TEXT("TEXTURE1DARRAY_VAR"));
+		ReplacementsForInvalid.Add(TEXT("Texture2D"), TEXT("TEXTURE2D_VAR"));
+		ReplacementsForInvalid.Add(TEXT("Texture2DArray"), TEXT("TEXTURE2DARRAY_VAR"));
+		ReplacementsForInvalid.Add(TEXT("Texture2DMS"), TEXT("TEXTURE2DMS_VAR"));
+		ReplacementsForInvalid.Add(TEXT("Texture2DMSArray"), TEXT("TEXTURE2DMSARRAY_VAR"));
+		ReplacementsForInvalid.Add(TEXT("Texture3D"), TEXT("TEXTURE3D_VAR"));
+		ReplacementsForInvalid.Add(TEXT("TextureCube"), TEXT("TEXTURECUBE_VAR"));
+		ReplacementsForInvalid.Add(TEXT("TextureCubeArray"), TEXT("TEXTURECUBEARRAY_VAR"));
+		ReplacementsForInvalid.Add(TEXT("true"), TEXT("TRUE_VAR"));
+		ReplacementsForInvalid.Add(TEXT("typedef"), TEXT("TYPEDEF_VAR"));
+		ReplacementsForInvalid.Add(TEXT("triangle"), TEXT("TRIANGLE_VAR"));
+		ReplacementsForInvalid.Add(TEXT("triangleadj"), TEXT("TRIANGLEADJ_VAR"));
+		ReplacementsForInvalid.Add(TEXT("TriangleStream"), TEXT("TRIANGLESTREAM_VAR"));
+		ReplacementsForInvalid.Add(TEXT("uint"), TEXT("UINT_VAR"));
+		ReplacementsForInvalid.Add(TEXT("uint1"), TEXT("UINT1_VAR"));
+		ReplacementsForInvalid.Add(TEXT("uint2"), TEXT("UINT2_VAR"));
+		ReplacementsForInvalid.Add(TEXT("uint3"), TEXT("UINT3_VAR"));
+		ReplacementsForInvalid.Add(TEXT("uint4"), TEXT("UINT4_VAR"));
+		ReplacementsForInvalid.Add(TEXT("uniform"), TEXT("UNIFORM_VAR"));
+		ReplacementsForInvalid.Add(TEXT("unorm"), TEXT("UNORM_VAR"));
+		ReplacementsForInvalid.Add(TEXT("unsigned"), TEXT("UNSIGNED_VAR"));
+		ReplacementsForInvalid.Add(TEXT("vector"), TEXT("VECTOR_VAR"));
+		ReplacementsForInvalid.Add(TEXT("vertexfragment"), TEXT("VERTEXFRAGMENT_VAR"));
+		ReplacementsForInvalid.Add(TEXT("VertexShader"), TEXT("VERTEXSHADER_VAR"));
+		ReplacementsForInvalid.Add(TEXT("void"), TEXT("VOID_VAR"));
+		ReplacementsForInvalid.Add(TEXT("volatile"), TEXT("VOLATILE_VAR"));
+		ReplacementsForInvalid.Add(TEXT("while"), TEXT("WHILE_VAR"));
+
+		// additional glsl keywords
+		ReplacementsForInvalid.Add(TEXT("attribute"), TEXT("ATTRIBUTE_VAR"));
+		ReplacementsForInvalid.Add(TEXT("varying"), TEXT("VARYING_VAR"));
+		ReplacementsForInvalid.Add(TEXT("layout"), TEXT("LAYOUT_VAR"));
+		ReplacementsForInvalid.Add(TEXT("mat2"), TEXT("MAT2_VAR"));
+		ReplacementsForInvalid.Add(TEXT("mat3"), TEXT("MAT3_VAR"));
+		ReplacementsForInvalid.Add(TEXT("mat4"), TEXT("MAT4_VAR"));
+		ReplacementsForInvalid.Add(TEXT("mat2x2"), TEXT("MAT2X2_VAR"));
+		ReplacementsForInvalid.Add(TEXT("mat2x3"), TEXT("MAT2X3_VAR"));
+		ReplacementsForInvalid.Add(TEXT("mat2x4"), TEXT("MAT2X4_VAR"));
+		ReplacementsForInvalid.Add(TEXT("mat3x2"), TEXT("MAT3X2_VAR"));
+		ReplacementsForInvalid.Add(TEXT("mat3x3"), TEXT("MAT3X3_VAR"));
+		ReplacementsForInvalid.Add(TEXT("mat3x4"), TEXT("MAT3X4_VAR"));
+		ReplacementsForInvalid.Add(TEXT("mat4x2"), TEXT("MAT4X2_VAR"));
+		ReplacementsForInvalid.Add(TEXT("mat4x3"), TEXT("MAT4X3_VAR"));
+		ReplacementsForInvalid.Add(TEXT("mat4x4"), TEXT("MAT4X4_VAR"));
+		ReplacementsForInvalid.Add(TEXT("vec2"), TEXT("VEC2_VAR"));
+		ReplacementsForInvalid.Add(TEXT("vec3"), TEXT("VEC3_VAR"));
+		ReplacementsForInvalid.Add(TEXT("vec4"), TEXT("VEC4_VAR"));
+		ReplacementsForInvalid.Add(TEXT("ivec2"), TEXT("IVEC2_VAR"));
+		ReplacementsForInvalid.Add(TEXT("ivec3"), TEXT("IVEC3_VAR"));
+		ReplacementsForInvalid.Add(TEXT("ivec4"), TEXT("IVEC4_VAR"));
+		ReplacementsForInvalid.Add(TEXT("bvec2"), TEXT("BVEC2_VAR"));
+		ReplacementsForInvalid.Add(TEXT("bvec3"), TEXT("BVEC3_VAR"));
+		ReplacementsForInvalid.Add(TEXT("bvec4"), TEXT("BVEC4_VAR"));
+		ReplacementsForInvalid.Add(TEXT("uvec2"), TEXT("UVEC2_VAR"));
+		ReplacementsForInvalid.Add(TEXT("uvec3"), TEXT("UVEC3_VAR"));
+		ReplacementsForInvalid.Add(TEXT("uvec4"), TEXT("UVEC4_VAR"));
+		ReplacementsForInvalid.Add(TEXT("lowp"), TEXT("LOWP_VAR"));
+		ReplacementsForInvalid.Add(TEXT("mediump"), TEXT("MEDIUMP_VAR"));
+		ReplacementsForInvalid.Add(TEXT("highp"), TEXT("HIGHP_VAR"));
+		ReplacementsForInvalid.Add(TEXT("precision"), TEXT("PRECISION_VAR"));
+		ReplacementsForInvalid.Add(TEXT("sampler1D"), TEXT("SAMPLER1D_VAR"));
+		ReplacementsForInvalid.Add(TEXT("sampler2D"), TEXT("SAMPLER2D_VAR"));
+		ReplacementsForInvalid.Add(TEXT("sampler3D"), TEXT("SAMPLER3D_VAR"));
+		ReplacementsForInvalid.Add(TEXT("samplerCube"), TEXT("SAMPLERCUBE_VAR"));
+		ReplacementsForInvalid.Add(TEXT("sampler1DShadow"), TEXT("SAMPLER1DSHADOW_VAR"));
+		ReplacementsForInvalid.Add(TEXT("sampler2DShadow"), TEXT("SAMPLER2DSHADOW_VAR"));
+		ReplacementsForInvalid.Add(TEXT("samplerCubeShadow"), TEXT("SAMPLERCUBESHADOW_VAR"));
+		ReplacementsForInvalid.Add(TEXT("sampler1DArray"), TEXT("SAMPLER1DARRAY_VAR"));
+		ReplacementsForInvalid.Add(TEXT("sampler2DArray"), TEXT("SAMPLER2DARRAY_VAR"));
+		ReplacementsForInvalid.Add(TEXT("sampler1DArrayShadow"), TEXT("SAMPLER1DARRAYSHADOW_VAR"));
+		ReplacementsForInvalid.Add(TEXT("sampler2DArrayShadow"), TEXT("SAMPLER2DARRAYSHADOW_VAR"));
+		ReplacementsForInvalid.Add(TEXT("isampler1D"), TEXT("ISAMPLER1D_VAR"));
+		ReplacementsForInvalid.Add(TEXT("isampler2D"), TEXT("ISAMPLER2D_VAR"));
+		ReplacementsForInvalid.Add(TEXT("isampler3D"), TEXT("ISAMPLER3D_VAR"));
+		ReplacementsForInvalid.Add(TEXT("isamplerCube"), TEXT("ISAMPLERCUBE_VAR"));
+		ReplacementsForInvalid.Add(TEXT("isampler1DArray"), TEXT("ISAMPLER1DARRAY_VAR"));
+		ReplacementsForInvalid.Add(TEXT("isampler2DArray"), TEXT("ISAMPLER2DARRAY_VAR"));
+		ReplacementsForInvalid.Add(TEXT("usampler1D"), TEXT("USAMPLER1D_VAR"));
+		ReplacementsForInvalid.Add(TEXT("usampler2D"), TEXT("USAMPLER2D_VAR"));
+		ReplacementsForInvalid.Add(TEXT("usampler3D"), TEXT("USAMPLER3D_VAR"));
+		ReplacementsForInvalid.Add(TEXT("usamplerCube"), TEXT("USAMPLERCUBE_VAR"));
+		ReplacementsForInvalid.Add(TEXT("usampler1DArray"), TEXT("USAMPLER1DARRAY_VAR"));
+		ReplacementsForInvalid.Add(TEXT("usampler2DArray"), TEXT("USAMPLER2DARRAY_VAR"));
+		ReplacementsForInvalid.Add(TEXT("sampler2DRect"), TEXT("SAMPLER2DRECT_VAR"));
+		ReplacementsForInvalid.Add(TEXT("sampler2DRectShadow"), TEXT("SAMPLER2DRECTSHADOW_VAR"));
+		ReplacementsForInvalid.Add(TEXT("isampler2DRect"), TEXT("ISAMPLER2DRECT_VAR"));
+		ReplacementsForInvalid.Add(TEXT("usampler2DRect"), TEXT("USAMPLER2DRECT_VAR"));
+		ReplacementsForInvalid.Add(TEXT("samplerBuffer"), TEXT("SAMPLERBUFFER_VAR"));
+		ReplacementsForInvalid.Add(TEXT("isamplerBuffer"), TEXT("ISAMPLERBUFFER_VAR"));
+		ReplacementsForInvalid.Add(TEXT("usamplerBuffer"), TEXT("USAMPLERBUFFER_VAR"));
+		ReplacementsForInvalid.Add(TEXT("struct"), TEXT("STRUCT_VAR"));
+		ReplacementsForInvalid.Add(TEXT("flat"), TEXT("FLAT_VAR"));
+		ReplacementsForInvalid.Add(TEXT("smooth"), TEXT("SMOOTH_VAR"));
+		ReplacementsForInvalid.Add(TEXT("invariant"), TEXT("INVARIANT_VAR"));
+		ReplacementsForInvalid.Add(TEXT("union"), TEXT("UNION_VAR"));
+		ReplacementsForInvalid.Add(TEXT("enum"), TEXT("ENUM_VAR"));
+		ReplacementsForInvalid.Add(TEXT("typedef"), TEXT("TYPEDEF_VAR"));
+		ReplacementsForInvalid.Add(TEXT("template"), TEXT("TEMPLATE_VAR"));
+		ReplacementsForInvalid.Add(TEXT("this"), TEXT("THIS_VAR"));
+		ReplacementsForInvalid.Add(TEXT("packed"), TEXT("PACKED_VAR"));
+		ReplacementsForInvalid.Add(TEXT("goto"), TEXT("GOTO_VAR"));
+		ReplacementsForInvalid.Add(TEXT("noinline"), TEXT("NOINLINE_VAR"));
+		ReplacementsForInvalid.Add(TEXT("public"), TEXT("PUBLIC_VAR"));
+		ReplacementsForInvalid.Add(TEXT("external"), TEXT("EXTERNAL_VAR"));
+		ReplacementsForInvalid.Add(TEXT("long short"), TEXT("LONG SHORT_VAR"));
+		ReplacementsForInvalid.Add(TEXT("superp"), TEXT("SUPERP_VAR"));
+		ReplacementsForInvalid.Add(TEXT("input"), TEXT("INPUT_VAR"));
+		ReplacementsForInvalid.Add(TEXT("output"), TEXT("OUTPUT_VAR"));
+		ReplacementsForInvalid.Add(TEXT("hvec2"), TEXT("HVEC2_VAR"));
+		ReplacementsForInvalid.Add(TEXT("hvec3"), TEXT("HVEC3_VAR"));
+		ReplacementsForInvalid.Add(TEXT("hvec4"), TEXT("HVEC4_VAR"));
+		ReplacementsForInvalid.Add(TEXT("dvec2"), TEXT("DVEC2_VAR"));
+		ReplacementsForInvalid.Add(TEXT("dvec3"), TEXT("DVEC3_VAR"));
+		ReplacementsForInvalid.Add(TEXT("dvec4"), TEXT("DVEC4_VAR"));
+		ReplacementsForInvalid.Add(TEXT("fvec2"), TEXT("FVEC2_VAR"));
+		ReplacementsForInvalid.Add(TEXT("fvec3"), TEXT("FVEC3_VAR"));
+		ReplacementsForInvalid.Add(TEXT("fvec4"), TEXT("FVEC4_VAR"));
+		ReplacementsForInvalid.Add(TEXT("sampler3DRect"), TEXT("SAMPLER3DRECT_VAR"));
+		ReplacementsForInvalid.Add(TEXT("filter"), TEXT("FILTER_VAR"));
+		ReplacementsForInvalid.Add(TEXT("image1D"), TEXT("IMAGE1D_VAR"));
+		ReplacementsForInvalid.Add(TEXT("image2D"), TEXT("IMAGE2D_VAR"));
+		ReplacementsForInvalid.Add(TEXT("image3D"), TEXT("IMAGE3D_VAR"));
+		ReplacementsForInvalid.Add(TEXT("imageCube"), TEXT("IMAGECUBE_VAR"));
+		ReplacementsForInvalid.Add(TEXT("iimage1D"), TEXT("IIMAGE1D_VAR"));
+		ReplacementsForInvalid.Add(TEXT("iimage2D"), TEXT("IIMAGE2D_VAR"));
+		ReplacementsForInvalid.Add(TEXT("iimage3D"), TEXT("IIMAGE3D_VAR"));
+		ReplacementsForInvalid.Add(TEXT("iimageCub"), TEXT("IIMAGECUBE_VAR"));
+		ReplacementsForInvalid.Add(TEXT("uimage1D"), TEXT("UIMAGE1D_VAR"));
+		ReplacementsForInvalid.Add(TEXT("uimage2D"), TEXT("UIMAGE2D_VAR"));
+		ReplacementsForInvalid.Add(TEXT("uimage3D"), TEXT("UIMAGE3D_VAR"));
+		ReplacementsForInvalid.Add(TEXT("uimageCube"), TEXT("UIMAGECUBE_VAR"));
+		ReplacementsForInvalid.Add(TEXT("image1DArray"), TEXT("IMAGE1DARRAY_VAR"));
+		ReplacementsForInvalid.Add(TEXT("image2DArray"), TEXT("IMAGE2DARRAY_VAR"));
+		ReplacementsForInvalid.Add(TEXT("iimage1DArray"), TEXT("IIMAGE1DARRAY_VAR"));
+		ReplacementsForInvalid.Add(TEXT("iimage2DArray"), TEXT("IIMAGE2DARRAY_VAR"));
+		ReplacementsForInvalid.Add(TEXT("uimage1DArray"), TEXT("UIMAGE1DARRAY_VAR"));
+		ReplacementsForInvalid.Add(TEXT("uimage2DArray"), TEXT("UIMAGE2DARRAY_VAR"));
+		ReplacementsForInvalid.Add(TEXT("image1DShadow"), TEXT("IMAGE1DSHADOW_VAR"));
+		ReplacementsForInvalid.Add(TEXT("image2DShadow"), TEXT("IMAGE2DSHADOW_VAR"));
+		ReplacementsForInvalid.Add(TEXT("image1DArrayShadow"), TEXT("IMAGE1DARRAYSHADOW_VAR"));
+		ReplacementsForInvalid.Add(TEXT("image2DArrayShadow"), TEXT("IMAGE2DARRAYSHADOW_VAR"));
+		ReplacementsForInvalid.Add(TEXT("imageBuffer"), TEXT("IMAGEBUFFER_VAR"));
+		ReplacementsForInvalid.Add(TEXT("iimageBuffer"), TEXT("IIMAGEBUFFER_VAR"));
+		ReplacementsForInvalid.Add(TEXT("uimageBuffer"), TEXT("UIMAGEBUFFER_VAR"));
+		ReplacementsForInvalid.Add(TEXT("sizeof"), TEXT("SIZEOF_VAR"));
+		ReplacementsForInvalid.Add(TEXT("cast"), TEXT("CAST_VAR"));
+		ReplacementsForInvalid.Add(TEXT("using"), TEXT("USING_VAR"));
+		ReplacementsForInvalid.Add(TEXT("row_major"), TEXT("ROW_MAJOR_VAR"));
+		ReplacementsForInvalid.Add(TEXT("int8_t"), TEXT("INT8_T_VAR"));
+		ReplacementsForInvalid.Add(TEXT("unsigned char"), TEXT("UNSIGNED_CHAR_VAR"));
+		ReplacementsForInvalid.Add(TEXT("uchar"), TEXT("UCHAR_VAR"));
+		ReplacementsForInvalid.Add(TEXT("uint8_t"), TEXT("UINT8_T_VAR"));
+		ReplacementsForInvalid.Add(TEXT("short"), TEXT("SHORT_VAR"));
+		ReplacementsForInvalid.Add(TEXT("int16_t"), TEXT("INT16_T_VAR"));
+		ReplacementsForInvalid.Add(TEXT("unsigned shorT"), TEXT("UNSIGNED_SHORT_VAR"));
+		ReplacementsForInvalid.Add(TEXT("ushort"), TEXT("USHORT_VAR"));
+		ReplacementsForInvalid.Add(TEXT("unit16_t"), TEXT("UNIT16_T_VAR"));
+		ReplacementsForInvalid.Add(TEXT("int32_t"), TEXT("INT32_T_VAR"));
+		ReplacementsForInvalid.Add(TEXT("unsigned int"), TEXT("UNSIGNED_INT_VAR"));
+		ReplacementsForInvalid.Add(TEXT("uint32_t"), TEXT("UINT32_T_VAR"));
+		ReplacementsForInvalid.Add(TEXT("smooth"), TEXT("SMOOTH_VAR"));
+		ReplacementsForInvalid.Add(TEXT("invariant"), TEXT("INVARIANT_VAR"));
+		ReplacementsForInvalid.Add(TEXT("union"), TEXT("UNION_VAR"));
+		ReplacementsForInvalid.Add(TEXT("enum"), TEXT("ENUM_VAR"));
+		ReplacementsForInvalid.Add(TEXT("typedef"), TEXT("TYPEDEF_VAR"));
+		ReplacementsForInvalid.Add(TEXT("template"), TEXT("TEMPLATE_VAR"));
+		ReplacementsForInvalid.Add(TEXT("this"), TEXT("THIS_VAR"));
+		ReplacementsForInvalid.Add(TEXT("packed"), TEXT("PACKED_VAR"));
+		ReplacementsForInvalid.Add(TEXT("goto"), TEXT("GOTO_VAR"));
+		ReplacementsForInvalid.Add(TEXT("noinline"), TEXT("NOINLINE_VAR"));
+		ReplacementsForInvalid.Add(TEXT("public"), TEXT("PUBLIC_VAR"));
+		ReplacementsForInvalid.Add(TEXT("external"), TEXT("EXTERNAL_VAR"));
+		ReplacementsForInvalid.Add(TEXT(" long short"), TEXT("LONG_SHORT_VAR"));
+		ReplacementsForInvalid.Add(TEXT("superp"), TEXT("SUPERP_VAR"));
+		ReplacementsForInvalid.Add(TEXT("input"), TEXT("INPUT_VAR"));
+		ReplacementsForInvalid.Add(TEXT("output"), TEXT("OUTPUT_VAR"));
+		ReplacementsForInvalid.Add(TEXT("hvec2"), TEXT("HVEC2_VAR"));
+		ReplacementsForInvalid.Add(TEXT("hvec3"), TEXT("HVEC3_VAR"));
+		ReplacementsForInvalid.Add(TEXT("hvec4"), TEXT("HVEC4_VAR"));
+		ReplacementsForInvalid.Add(TEXT("dvec2"), TEXT("DVEC2_VAR"));
+		ReplacementsForInvalid.Add(TEXT("dvec3"), TEXT("DVEC3_VAR"));
+		ReplacementsForInvalid.Add(TEXT("dvec4"), TEXT("DVEC4_VAR"));
+		ReplacementsForInvalid.Add(TEXT("fvec2"), TEXT("FVEC2_VAR"));
+		ReplacementsForInvalid.Add(TEXT("fvec3"), TEXT("FVEC3_VAR"));
+		ReplacementsForInvalid.Add(TEXT("fvec4"), TEXT("FVEC4_VAR"));
+		ReplacementsForInvalid.Add(TEXT("sampler3DRect"), TEXT("SAMPLER3DRECT_VAR"));
+		ReplacementsForInvalid.Add(TEXT("filter"), TEXT("FILTER_VAR"));
+		ReplacementsForInvalid.Add(TEXT("image2D"), TEXT("IMAGE2D_VAR"));
+		ReplacementsForInvalid.Add(TEXT("image3D"), TEXT("IMAGE3D_VAR"));
+		ReplacementsForInvalid.Add(TEXT("imageCube"), TEXT("IMAGECUBE_VAR"));
+		ReplacementsForInvalid.Add(TEXT("iimage1D"), TEXT("IIMAGE1D_VAR"));
+		ReplacementsForInvalid.Add(TEXT("iimage2D"), TEXT("IIMAGE2D_VAR"));
+		ReplacementsForInvalid.Add(TEXT("iimage3D"), TEXT("IIMAGE3D_VAR"));
+		ReplacementsForInvalid.Add(TEXT("iimageCube"), TEXT("IIMAGECUBE_VAR"));
+		ReplacementsForInvalid.Add(TEXT("uimage1D"), TEXT("UIMAGE1D_VAR"));
+		ReplacementsForInvalid.Add(TEXT("uimage2D"), TEXT("UIMAGE2D_VAR"));
+		ReplacementsForInvalid.Add(TEXT("uimage3D"), TEXT("UIMAGE3D_VAR"));
+		ReplacementsForInvalid.Add(TEXT("uimageCube"), TEXT("UIMAGECUBE_VAR"));
+		ReplacementsForInvalid.Add(TEXT("image1DArray"), TEXT("IMAGE1DARRAY_VAR"));
+		ReplacementsForInvalid.Add(TEXT("image2DArray"), TEXT("IMAGE2DARRAY_VAR"));
+		ReplacementsForInvalid.Add(TEXT("iimage1DArray"), TEXT("IIMAGE1DARRAY_VAR"));
+		ReplacementsForInvalid.Add(TEXT("iimage2DArray"), TEXT("IIMAGE2DARRAY_VAR"));
+		ReplacementsForInvalid.Add(TEXT("uimage1DArray"), TEXT("UIMAGE1DARRAY_VAR"));
+		ReplacementsForInvalid.Add(TEXT("uimage2DArray"), TEXT("UIMAGE2DARRAY_VAR"));
+		ReplacementsForInvalid.Add(TEXT("image1DShadow"), TEXT("IMAGE1DSHADOW_VAR"));
+		ReplacementsForInvalid.Add(TEXT("image2DShadow"), TEXT("IMAGE2DSHADOW_VAR"));
+		ReplacementsForInvalid.Add(TEXT("image1DArrayShadow"), TEXT("IMAGE1DARRAYSHADOW_VAR"));
+		ReplacementsForInvalid.Add(TEXT("image2DArrayShadow"), TEXT("IMAGE2DARRAYSHADOW_VAR"));
+		ReplacementsForInvalid.Add(TEXT("imageBuffer"), TEXT("IMAGEBUFFER_VAR"));
+		ReplacementsForInvalid.Add(TEXT("iimageBuffer"), TEXT("IIMAGEBUFFER_VAR"));
+		ReplacementsForInvalid.Add(TEXT("uimageBuffer"), TEXT("UIMAGEBUFFER_VAR"));
+		ReplacementsForInvalid.Add(TEXT("sizeof"), TEXT("SIZEOF_VAR"));
+		ReplacementsForInvalid.Add(TEXT("cast"), TEXT("CAST_VAR"));
+		ReplacementsForInvalid.Add(TEXT("using"), TEXT("USING_VAR"));
+		ReplacementsForInvalid.Add(TEXT("row_major"), TEXT("ROW_MAJOR_VAR"));
+		ReplacementsForInvalid.Add(TEXT("int8_t"), TEXT("INT8_T_VAR"));
+		ReplacementsForInvalid.Add(TEXT("unsigned char"), TEXT("UNSIGNED_CHAR_VAR"));
+		ReplacementsForInvalid.Add(TEXT("uchar"), TEXT("UCHAR_VAR"));
+		ReplacementsForInvalid.Add(TEXT("uint8_t"), TEXT("UINT8_T_VAR"));
+		ReplacementsForInvalid.Add(TEXT("short"), TEXT("SHORT_VAR"));
+		ReplacementsForInvalid.Add(TEXT("int16_t"), TEXT("INT16_T_VAR"));
+		ReplacementsForInvalid.Add(TEXT("unsigned short"), TEXT("UNSIGNED_SHORT_VAR"));
+		ReplacementsForInvalid.Add(TEXT("ushort"), TEXT("USHORT_VAR"));
+		ReplacementsForInvalid.Add(TEXT("unit16_t"), TEXT("UNIT16_T_VAR"));
+		ReplacementsForInvalid.Add(TEXT("int32_t"), TEXT("INT32_T_VAR"));
+		ReplacementsForInvalid.Add(TEXT("unsigned int"), TEXT("UNSIGNED_INT_VAR"));
+		ReplacementsForInvalid.Add(TEXT("uint32_t"), TEXT("UINT32_T_VAR"));
+	}
+}
+
 FString FHlslNiagaraTranslator::GetSanitizedSymbolName(FString SymbolName, bool bCollapsNamespaces)
 {
+	if (SymbolName.Len() == 0)
+	{
+		return SymbolName;
+	}
+
+	check(ReplacementsForInvalid.Num() != 0);
+
 	FString Ret = SymbolName;
-	Ret.ReplaceInline(TEXT(" "), TEXT(""));
+
+	// Split up into individual namespaces...
+	TArray<FString> SplitName;
+	Ret.ParseIntoArray(SplitName, TEXT("."));
+
+	// Rules for variable namespaces..
+	for (int32 i = 0; i < SplitName.Num(); i++)
+	{
+		SplitName[i].ReplaceInline(TEXT("__"), TEXT("ASC95ASC95")); // Opengl reserves "__" within a name
+		if (SplitName[i][0] >= TCHAR('0') && SplitName[i][0] <= TCHAR('9')) // Cannot start with a numeric digit
+		{
+			SplitName[i] = TEXT("INTEGER_") + SplitName[i];
+		}
+		
+		FString* FoundReplacementStr = ReplacementsForInvalid.Find(SplitName[i]); // Look for the string in the keyword protections array.
+		if (FoundReplacementStr)
+		{
+			SplitName[i] = *FoundReplacementStr;
+		}		
+
+		SplitName[i].ReplaceInline(TEXT("\t"), TEXT(""));
+		SplitName[i].ReplaceInline(TEXT(" "), TEXT(""));
+
+		// Handle internationalization of characters..
+		FString ChangedSplitName;
+		ChangedSplitName.Reserve(SplitName[i].Len()*6); // Assign room for every current char to be 'ASCXXX'
+		for (int32 j = 0; j < SplitName[i].Len(); j++)
+		{
+			if ((SplitName[i][j] >= TCHAR('0') && SplitName[i][j] <= TCHAR('9')) ||
+				(SplitName[i][j] >= TCHAR('A') && SplitName[i][j] <= TCHAR('Z')) ||
+				(SplitName[i][j] >= TCHAR('a') && SplitName[i][j] <= TCHAR('z')) ||
+				SplitName[i][j] == TCHAR('_') ||
+				SplitName[i][j] == TCHAR(' '))
+			{
+				// Do nothing.. these are valid chars..
+				ChangedSplitName.AppendChar(SplitName[i][j]);
+			}
+			else
+			{
+				// Need to replace the bad characters..
+				ChangedSplitName.Append(TEXT("ASC"));
+				ChangedSplitName.AppendInt((int32)SplitName[i][j]);
+			}
+		}
+		SplitName[i] = ChangedSplitName;
+	}
+
+	// Gather back into single string..
+	Ret = FString::Join<FString>(SplitName, TEXT("."));
+
+	/*
 	Ret.ReplaceInline(TEXT("\\"), TEXT("_"));
 	Ret.ReplaceInline(TEXT("/"), TEXT("_"));
 	Ret.ReplaceInline(TEXT(","), TEXT("_"));
 	Ret.ReplaceInline(TEXT("-"), TEXT("_"));
 	Ret.ReplaceInline(TEXT(":"), TEXT("_"));
+	Ret = Ret.ConvertTabsToSpaces(0);
+	*/
 
 	if (bCollapsNamespaces)
 	{
@@ -2238,7 +2656,7 @@ int32 FHlslNiagaraTranslator::AddUniformChunk(FString SymbolName, const FNiagara
 	return Ret;
 }
 
-int32 FHlslNiagaraTranslator::AddSourceChunk(FString SymbolName, const FNiagaraTypeDefinition& Type)
+int32 FHlslNiagaraTranslator::AddSourceChunk(FString SymbolName, const FNiagaraTypeDefinition& Type, bool bSanitize)
 {
 	int32 Ret = CodeChunks.IndexOfByPredicate(
 		[&](const FNiagaraCodeChunk& Chunk)
@@ -2251,7 +2669,7 @@ int32 FHlslNiagaraTranslator::AddSourceChunk(FString SymbolName, const FNiagaraT
 	{
 		Ret = CodeChunks.AddDefaulted();
 		FNiagaraCodeChunk& Chunk = CodeChunks[Ret];
-		Chunk.SymbolName = GetSanitizedSymbolName(SymbolName);
+		Chunk.SymbolName = bSanitize ? GetSanitizedSymbolName(SymbolName) : SymbolName;
 		Chunk.Type = Type;
 		Chunk.Mode = ENiagaraCodeChunkMode::Source;
 
@@ -2752,7 +3170,7 @@ void FHlslNiagaraTranslator::Output(UNiagaraNodeOutput* OutputNode, const TArray
 
 			if (Attributes[i].GetType() != FNiagaraTypeDefinition::GetParameterMapDef())
 			{
-				FString SymbolName = *(TEXT("Out_") + GetSanitizedSymbolName(Attributes[i].GetName().ToString()));
+				FString SymbolName = *GetSanitizedSymbolName((TEXT("Out_") + Attributes[i].GetName().ToString()));
 				ENiagaraCodeChunkMode OldMode = CurrentBodyChunkMode;
 				CurrentBodyChunkMode = ENiagaraCodeChunkMode::Body;
 				AddBodyChunk(SymbolName, TEXT("{0}"), Attributes[i].GetType(), Inputs[i], false);
@@ -4050,7 +4468,7 @@ void FHlslNiagaraTranslator::FinalResolveNamespacedTokens(const FString& Paramet
 					}
 
 					
-					Tokens[i] = ParameterMapInstanceNamespace + Var.GetName().ToString();
+					Tokens[i] = ParameterMapInstanceNamespace + GetSanitizedSymbolName(Var.GetName().ToString());
 					break;
 				}
 			}
@@ -4444,7 +4862,7 @@ void FHlslNiagaraTranslator::RegisterFunctionCall(ENiagaraScriptUsage ScriptUsag
 				if (Info.UserPtrIdx != INDEX_NONE)
 				{
 					//This interface requires per instance data via a user ptr so place the index to it at the end of the inputs.
-					Inputs.Add(AddSourceChunk(LexToString(Info.UserPtrIdx), FNiagaraTypeDefinition::GetIntDef()));
+					Inputs.Add(AddSourceChunk(LexToString(Info.UserPtrIdx), FNiagaraTypeDefinition::GetIntDef(), false));
 					OutSignature.Inputs.Add(FNiagaraVariable(FNiagaraTypeDefinition::GetIntDef(), TEXT("InstanceData")));
 				}
 			}
@@ -4652,7 +5070,7 @@ FString FHlslNiagaraTranslator::GetFunctionSignature(const FNiagaraFunctionSigna
 					SigStr += TEXT(", ");
 				}
 
-				SigStr += TEXT("out ") + FHlslNiagaraTranslator::GetStructHlslTypeName(Output.GetType()) + TEXT(" Out_") + FHlslNiagaraTranslator::GetSanitizedSymbolName(Output.GetName().ToString());
+				SigStr += TEXT("out ") + FHlslNiagaraTranslator::GetStructHlslTypeName(Output.GetType()) + TEXT(" ") + FHlslNiagaraTranslator::GetSanitizedSymbolName(TEXT("Out_") + Output.GetName().ToString());
 				++ParamIdx;
 			}
 		}
diff --git a/Engine/Plugins/FX/Niagara/Source/NiagaraEditor/Private/NiagaraHlslTranslator.h b/Engine/Plugins/FX/Niagara/Source/NiagaraEditor/Private/NiagaraHlslTranslator.h
--- a/Engine/Plugins/FX/Niagara/Source/NiagaraEditor/Private/NiagaraHlslTranslator.h
+++ b/Engine/Plugins/FX/Niagara/Source/NiagaraEditor/Private/NiagaraHlslTranslator.h
@@ -381,7 +381,7 @@ protected:
 	FString GeneratedConstantString(FVector4 Constant);
 
 	/* Add a chunk that is not written to the source, only used as a source chunk for others. */
-	int32 AddSourceChunk(FString SymbolName, const FNiagaraTypeDefinition& Type);
+	int32 AddSourceChunk(FString SymbolName, const FNiagaraTypeDefinition& Type, bool bSanitize = true);
 
 	/** Add a chunk defining a uniform value. */
 	int32 AddUniformChunk(FString SymbolName, const FNiagaraTypeDefinition& Type);
@@ -401,6 +401,8 @@ public:
 
 	FHlslNiagaraTranslator();
 	virtual ~FHlslNiagaraTranslator() {}
+
+	static void Init();
 	
 	virtual const FNiagaraTranslateResults &Translate(const FNiagaraCompileRequestData* InCompileData, const FNiagaraCompileOptions& InCompileOptions, FHlslNiagaraTranslatorOptions Options);
 	FNiagaraTranslatorOutput &GetTranslateOutput() { return CompilationOutput; }
@@ -614,4 +616,7 @@ private:
 	TArray<FNiagaraVariable> InitialNamespaceVariablesMissingDefault;
 	// Variables that need to be initialized in the body or at the end of spawn.
 	TArray<FNiagaraVariable> DeferredVariablesMissingDefault;
+
+	static TMap<FString, FString> ReplacementsForInvalid;
+
 };
diff --git a/Engine/Plugins/FX/Niagara/Source/NiagaraEditor/Private/NiagaraNodeCustomHlsl.cpp b/Engine/Plugins/FX/Niagara/Source/NiagaraEditor/Private/NiagaraNodeCustomHlsl.cpp
--- a/Engine/Plugins/FX/Niagara/Source/NiagaraEditor/Private/NiagaraNodeCustomHlsl.cpp
+++ b/Engine/Plugins/FX/Niagara/Source/NiagaraEditor/Private/NiagaraNodeCustomHlsl.cpp
@@ -39,7 +39,7 @@ FText UNiagaraNodeCustomHlsl::GetHlslText() const
 void UNiagaraNodeCustomHlsl::OnCustomHlslTextCommitted(const FText& InText, ETextCommit::Type InType)
 {
 	FString NewValue = InText.ToString();
-	if (NewValue != CustomHlsl)
+	if (!NewValue.Equals(CustomHlsl, ESearchCase::CaseSensitive))
 	{
 		FScopedTransaction Transaction(LOCTEXT("CustomHlslCommit", "Edited Custom Hlsl"));
 		Modify();
diff --git a/Engine/Plugins/FX/Niagara/Source/NiagaraEditor/Private/NiagaraNodeInput.cpp b/Engine/Plugins/FX/Niagara/Source/NiagaraEditor/Private/NiagaraNodeInput.cpp
--- a/Engine/Plugins/FX/Niagara/Source/NiagaraEditor/Private/NiagaraNodeInput.cpp
+++ b/Engine/Plugins/FX/Niagara/Source/NiagaraEditor/Private/NiagaraNodeInput.cpp
@@ -435,7 +435,12 @@ void UNiagaraNodeInput::Compile(class FHlslNiagaraTranslator* Translator, TArray
 	}
 
 	UNiagaraGraph* Graph = GetNiagaraGraph();
-	check(Input.GetType() != FNiagaraTypeDefinition::GetGenericNumericDef());
+	if (Input.GetType() == FNiagaraTypeDefinition::GetGenericNumericDef())
+	{
+		Outputs.Add(INDEX_NONE);
+		Translator->Error(LOCTEXT("InvalidPinType", "Numeric types should be able to be inferred from use by this phase of compilation."), this, nullptr);
+		return;
+	}
 
 	int32 FunctionParam = INDEX_NONE;
 	if (IsExposed() && Translator->GetFunctionParameter(Input, FunctionParam))
diff --git a/Engine/Plugins/FX/Niagara/Source/NiagaraEditor/Private/NiagaraParameterMapHistory.cpp b/Engine/Plugins/FX/Niagara/Source/NiagaraEditor/Private/NiagaraParameterMapHistory.cpp
--- a/Engine/Plugins/FX/Niagara/Source/NiagaraEditor/Private/NiagaraParameterMapHistory.cpp
+++ b/Engine/Plugins/FX/Niagara/Source/NiagaraEditor/Private/NiagaraParameterMapHistory.cpp
@@ -250,15 +250,17 @@ FString FNiagaraParameterMapHistory::MakeSafeNamespaceString(const FString& InSt
 }
 
 
-FNiagaraVariable FNiagaraParameterMapHistory::ResolveAsBasicAttribute(const FNiagaraVariable& InVar)
+FNiagaraVariable FNiagaraParameterMapHistory::ResolveAsBasicAttribute(const FNiagaraVariable& InVar, bool bSanitizeString)
 {
 	if (IsAttribute(InVar))
 	{
 		FString ParamName = InVar.GetName().ToString();
 		ParamName.RemoveAt(0, FString(PARAM_MAP_ATTRIBUTE_STR).Len());
 
-		ParamName = MakeSafeNamespaceString(ParamName);
-
+		if (bSanitizeString)
+		{
+			ParamName = MakeSafeNamespaceString(ParamName);
+		}
 		FNiagaraVariable RetVar = InVar;
 		RetVar.SetName(*ParamName);
 		return RetVar;
@@ -269,12 +271,15 @@ FNiagaraVariable FNiagaraParameterMapHistory::ResolveAsBasicAttribute(const FNia
 	}
 }
 
-FNiagaraVariable FNiagaraParameterMapHistory::BasicAttributeToNamespacedAttribute(const FNiagaraVariable& InVar)
+FNiagaraVariable FNiagaraParameterMapHistory::BasicAttributeToNamespacedAttribute(const FNiagaraVariable& InVar, bool bSanitizeString)
 {
 	FString ParamName = InVar.GetName().ToString();
 	ParamName.InsertAt(0, FString(PARAM_MAP_ATTRIBUTE_STR));
 
-	ParamName = MakeSafeNamespaceString(ParamName);
+	if (bSanitizeString)
+	{
+		ParamName = MakeSafeNamespaceString(ParamName);
+	}
 
 	FNiagaraVariable RetVar = InVar;
 	RetVar.SetName(*ParamName);
@@ -1188,11 +1193,11 @@ int32 FNiagaraParameterMapHistoryBuilder::HandleVariableRead(int32 ParamMapIdx,
 			{
 				FoundIdx = Histories[ParamMapIdx].Variables.Add(Var);
 				Histories[ParamMapIdx].VariablesWithOriginalAliasesIntact.Add(AliasedVar);
+				Histories[ParamMapIdx].PerVariableWarnings.AddDefaulted(1);
+				Histories[ParamMapIdx].PerVariableWriteHistory.AddDefaulted(1);
+				Histories[ParamMapIdx].PerVariableReadHistory.AddDefaulted(1);
 
 			}
-			Histories[ParamMapIdx].PerVariableWarnings.AddDefaulted(1);
-			Histories[ParamMapIdx].PerVariableWriteHistory.AddDefaulted(1);
-			Histories[ParamMapIdx].PerVariableReadHistory.AddDefaulted(1);
 			Histories[ParamMapIdx].PerVariableReadHistory[FoundIdx].Add(TTuple<const UEdGraphPin*, const UEdGraphPin*>(InPin, nullptr));
 
 			if (InDefaultPin && OutUsedDefault)
diff --git a/Engine/Plugins/FX/Niagara/Source/NiagaraEditor/Private/ViewModels/NiagaraEmitterHandleViewModel.cpp b/Engine/Plugins/FX/Niagara/Source/NiagaraEditor/Private/ViewModels/NiagaraEmitterHandleViewModel.cpp
--- a/Engine/Plugins/FX/Niagara/Source/NiagaraEditor/Private/ViewModels/NiagaraEmitterHandleViewModel.cpp
+++ b/Engine/Plugins/FX/Niagara/Source/NiagaraEditor/Private/ViewModels/NiagaraEmitterHandleViewModel.cpp
@@ -227,7 +227,11 @@ void FNiagaraEmitterHandleViewModel::OpenSourceEmitter()
 {
 	if (EmitterHandle)
 	{
-		FAssetEditorManager::Get().OpenEditorForAsset(const_cast<UNiagaraEmitter*>(EmitterHandle->GetSource()));
+		UNiagaraEmitter* EmitterSource = const_cast<UNiagaraEmitter*>(EmitterHandle->GetSource());
+		if (EmitterSource)
+		{
+			FAssetEditorManager::Get().OpenEditorForAsset(EmitterSource);
+		}
 	}
 }
 
diff --git a/Engine/Plugins/FX/Niagara/Source/NiagaraEditor/Private/ViewModels/Stack/NiagaraStackGraphUtilities.cpp b/Engine/Plugins/FX/Niagara/Source/NiagaraEditor/Private/ViewModels/Stack/NiagaraStackGraphUtilities.cpp
--- a/Engine/Plugins/FX/Niagara/Source/NiagaraEditor/Private/ViewModels/Stack/NiagaraStackGraphUtilities.cpp
+++ b/Engine/Plugins/FX/Niagara/Source/NiagaraEditor/Private/ViewModels/Stack/NiagaraStackGraphUtilities.cpp
@@ -616,54 +616,57 @@ void FNiagaraStackGraphUtilities::RemoveNodesForStackFunctionInputOverridePin(UE
 		{
 			UNiagaraNodeFunctionCall* DynamicInputNode = CastChecked<UNiagaraNodeFunctionCall>(OverrideValueNode);
 			UEdGraphPin* DynamicInputNodeInputPin = FNiagaraStackGraphUtilities::GetParameterMapInputPin(*DynamicInputNode);
-			UNiagaraNodeParameterMapSet* DynamicInputNodeOverrideNode = Cast<UNiagaraNodeParameterMapSet>(DynamicInputNodeInputPin->LinkedTo[0]->GetOwningNode());
-			if (DynamicInputNodeOverrideNode != nullptr)
+			if (DynamicInputNodeInputPin && DynamicInputNodeInputPin->LinkedTo.Num() > 0 && DynamicInputNodeInputPin->LinkedTo[0] != nullptr)
 			{
-				TArray<UEdGraphPin*> InputPins;
-				DynamicInputNodeOverrideNode->GetInputPins(InputPins);
-				for (UEdGraphPin* InputPin : InputPins)
+				UNiagaraNodeParameterMapSet* DynamicInputNodeOverrideNode = Cast<UNiagaraNodeParameterMapSet>(DynamicInputNodeInputPin->LinkedTo[0]->GetOwningNode());
+				if (DynamicInputNodeOverrideNode != nullptr)
 				{
-					if (InputPin->PinName.ToString().StartsWith(DynamicInputNode->GetFunctionName()))
+					TArray<UEdGraphPin*> InputPins;
+					DynamicInputNodeOverrideNode->GetInputPins(InputPins);
+					for (UEdGraphPin* InputPin : InputPins)
 					{
-						RemoveNodesForStackFunctionInputOverridePin(*InputPin, OutRemovedDataObjects);
-						DynamicInputNodeOverrideNode->RemovePin(InputPin);
+						if (InputPin->PinName.ToString().StartsWith(DynamicInputNode->GetFunctionName()))
+						{
+							RemoveNodesForStackFunctionInputOverridePin(*InputPin, OutRemovedDataObjects);
+							DynamicInputNodeOverrideNode->RemovePin(InputPin);
+						}
 					}
-				}
 
-				TArray<UEdGraphPin*> NewInputPins;
-				DynamicInputNodeOverrideNode->GetInputPins(NewInputPins);
-				if (NewInputPins.Num() == 2)
-				{
-					// If there are only 2 input pins left, they are the parameter map input and the add pin, so the dynamic input's override node 
-					// can be removed.  This not always be the case when removing dynamic input nodes because they share the same override node.
-					UEdGraphPin* InputPin = FNiagaraStackGraphUtilities::GetParameterMapInputPin(*DynamicInputNodeOverrideNode);
-					UEdGraphPin* OutputPin = FNiagaraStackGraphUtilities::GetParameterMapOutputPin(*DynamicInputNodeOverrideNode);
-
-					if (ensureMsgf(InputPin != nullptr && InputPin->LinkedTo.Num() == 1 && OutputPin != nullptr && 
-						OutputPin->LinkedTo.Num() >= 2,	TEXT("Invalid Stack - Dynamic input node override node not connected correctly.")))
+					TArray<UEdGraphPin*> NewInputPins;
+					DynamicInputNodeOverrideNode->GetInputPins(NewInputPins);
+					if (NewInputPins.Num() == 2)
 					{
-						// The DynamicInputOverrideNode will have a single input which is the previous module or override map set, and
-						// two or more output links, one to the dynamic input node, one to the next override map set, and 0 or more links
-						// to other dynamic inputs on sibling inputs.  Collect these linked pins to reconnect after removing the override node.
-						UEdGraphPin* LinkedInputPin = InputPin->LinkedTo[0];
-						TArray<UEdGraphPin*> LinkedOutputPins;
-						for (UEdGraphPin* LinkedOutputPin : OutputPin->LinkedTo)
+						// If there are only 2 input pins left, they are the parameter map input and the add pin, so the dynamic input's override node 
+						// can be removed.  This not always be the case when removing dynamic input nodes because they share the same override node.
+						UEdGraphPin* InputPin = FNiagaraStackGraphUtilities::GetParameterMapInputPin(*DynamicInputNodeOverrideNode);
+						UEdGraphPin* OutputPin = FNiagaraStackGraphUtilities::GetParameterMapOutputPin(*DynamicInputNodeOverrideNode);
+
+						if (ensureMsgf(InputPin != nullptr && InputPin->LinkedTo.Num() == 1 && OutputPin != nullptr &&
+							OutputPin->LinkedTo.Num() >= 2, TEXT("Invalid Stack - Dynamic input node override node not connected correctly.")))
 						{
-							if (LinkedOutputPin->GetOwningNode() != DynamicInputNode)
+							// The DynamicInputOverrideNode will have a single input which is the previous module or override map set, and
+							// two or more output links, one to the dynamic input node, one to the next override map set, and 0 or more links
+							// to other dynamic inputs on sibling inputs.  Collect these linked pins to reconnect after removing the override node.
+							UEdGraphPin* LinkedInputPin = InputPin->LinkedTo[0];
+							TArray<UEdGraphPin*> LinkedOutputPins;
+							for (UEdGraphPin* LinkedOutputPin : OutputPin->LinkedTo)
 							{
-								LinkedOutputPins.Add(LinkedOutputPin);
+								if (LinkedOutputPin->GetOwningNode() != DynamicInputNode)
+								{
+									LinkedOutputPins.Add(LinkedOutputPin);
+								}
 							}
-						}
 
-						// Disconnect the override node and remove it.
-						InputPin->BreakAllPinLinks();
-						OutputPin->BreakAllPinLinks();
-						Graph->RemoveNode(DynamicInputNodeOverrideNode);
+							// Disconnect the override node and remove it.
+							InputPin->BreakAllPinLinks();
+							OutputPin->BreakAllPinLinks();
+							Graph->RemoveNode(DynamicInputNodeOverrideNode);
 
-						// Reconnect the pins which were connected to the removed override node.
-						for (UEdGraphPin* LinkedOutputPin : LinkedOutputPins)
-						{
-							LinkedInputPin->MakeLinkTo(LinkedOutputPin);
+							// Reconnect the pins which were connected to the removed override node.
+							for (UEdGraphPin* LinkedOutputPin : LinkedOutputPins)
+							{
+								LinkedInputPin->MakeLinkTo(LinkedOutputPin);
+							}
 						}
 					}
 				}
@@ -746,15 +749,24 @@ void FNiagaraStackGraphUtilities::SetDynamicInputForFunctionInput(UEdGraphPin& O
 	const UEdGraphSchema_Niagara* NiagaraSchema = GetDefault<UEdGraphSchema_Niagara>();
 
 	FNiagaraTypeDefinition InputType = NiagaraSchema->PinToTypeDefinition(&OverridePin);
-	checkf(FunctionCallInputPin != nullptr, TEXT("Dynamic Input function call did not have a parameter map input pin."));
-	checkf(FunctionCallOutputPins.Num() == 1 && NiagaraSchema->PinToTypeDefinition(FunctionCallOutputPins[0]) == InputType, TEXT("Invalid Stack Graph - Dynamic Input function did not have the correct typed output pin"));
 
-	UEdGraphPin* OverrideNodeInputPin = FNiagaraStackGraphUtilities::GetParameterMapInputPin(*OverrideNode);
-	UEdGraphPin* PreviousStackNodeOutputPin = OverrideNodeInputPin->LinkedTo[0];
-	checkf(PreviousStackNodeOutputPin != nullptr, TEXT("Invalid Stack Graph - No previous stack node."));
 
-	FunctionCallInputPin->MakeLinkTo(PreviousStackNodeOutputPin);
-	FunctionCallOutputPins[0]->MakeLinkTo(&OverridePin);
+	UEdGraphPin* OverrideNodeInputPin = FNiagaraStackGraphUtilities::GetParameterMapInputPin(*OverrideNode);
+	UEdGraphPin* PreviousStackNodeOutputPin = nullptr;
+	if (OverrideNodeInputPin != nullptr)
+	{
+		PreviousStackNodeOutputPin = OverrideNodeInputPin->LinkedTo[0];
+	}
+	
+	if (FunctionCallInputPin != nullptr && PreviousStackNodeOutputPin != nullptr)
+	{
+		FunctionCallInputPin->MakeLinkTo(PreviousStackNodeOutputPin);
+	}
+	
+	if (FunctionCallOutputPins.Num() >= 1 && FunctionCallOutputPins[0] != nullptr)
+	{
+		FunctionCallOutputPins[0]->MakeLinkTo(&OverridePin);
+	}
 
 	OutDynamicInputFunctionCall = FunctionCallNode;
 
diff --git a/Engine/Plugins/FX/Niagara/Source/NiagaraEditor/Public/EdGraphSchema_Niagara.h b/Engine/Plugins/FX/Niagara/Source/NiagaraEditor/Public/EdGraphSchema_Niagara.h
--- a/Engine/Plugins/FX/Niagara/Source/NiagaraEditor/Public/EdGraphSchema_Niagara.h
+++ b/Engine/Plugins/FX/Niagara/Source/NiagaraEditor/Public/EdGraphSchema_Niagara.h
@@ -83,6 +83,7 @@ class NIAGARAEDITOR_API UEdGraphSchema_Niagara : public UEdGraphSchema
 
 	TArray<TSharedPtr<FNiagaraSchemaAction_NewNode> > GetGraphContextActions(const UEdGraph* CurrentGraph, TArray<UObject*>& SelectedObjects, const UEdGraphPin* FromPin, UEdGraph* OwnerOfTemporaries) const;
 	void PromoteSinglePinToParameter(UEdGraphPin* SourcePin);
+	static bool CanPromoteSinglePinToParameter(const UEdGraphPin* SourcePin);
 	void ToggleNodeEnabledState(class UNiagaraNode* InNode) const;
 
 	/** 
diff --git a/Engine/Plugins/FX/Niagara/Source/NiagaraEditor/Public/NiagaraParameterMapHistory.h b/Engine/Plugins/FX/Niagara/Source/NiagaraEditor/Public/NiagaraParameterMapHistory.h
--- a/Engine/Plugins/FX/Niagara/Source/NiagaraEditor/Public/NiagaraParameterMapHistory.h
+++ b/Engine/Plugins/FX/Niagara/Source/NiagaraEditor/Public/NiagaraParameterMapHistory.h
@@ -105,13 +105,13 @@ public:
 	/**
 	* Remove the Particles namespace if it exists.
 	*/
-	static FNiagaraVariable ResolveAsBasicAttribute(const FNiagaraVariable& InVar);
+	static FNiagaraVariable ResolveAsBasicAttribute(const FNiagaraVariable& InVar, bool bSanitizeInput = true);
 
 
 	/**
 	* Reverses ResolveAsBasicAttribute.
 	*/
-	static FNiagaraVariable BasicAttributeToNamespacedAttribute(const FNiagaraVariable& InVar);
+	static FNiagaraVariable BasicAttributeToNamespacedAttribute(const FNiagaraVariable& InVar, bool bSanitizeInput = true);
 
 	
 	/** Prepends the namespace string to the variable name.*/
diff --git a/Engine/Plugins/FX/Niagara/Source/NiagaraEditorWidgets/Private/SNiagaraStack.cpp b/Engine/Plugins/FX/Niagara/Source/NiagaraEditorWidgets/Private/SNiagaraStack.cpp
--- a/Engine/Plugins/FX/Niagara/Source/NiagaraEditorWidgets/Private/SNiagaraStack.cpp
+++ b/Engine/Plugins/FX/Niagara/Source/NiagaraEditorWidgets/Private/SNiagaraStack.cpp
@@ -135,7 +135,7 @@ void SNiagaraStack::Construct(const FArguments& InArgs, UNiagaraStackViewModel*
 					LOCTEXT("OpenAndFocusEmitter", "Open and focus source emitter"),
 					LOCTEXT("OpenAndFocusEmitterToolTip", "Open the source emitter in emitter mode"),
 					FSlateIcon(),
-					FUIAction(FExecuteAction::CreateSP(this, &SNiagaraStack::OpenSourceEmitter)));
+					FUIAction(FExecuteAction::CreateSP(this, &SNiagaraStack::OpenSourceEmitter), FCanExecuteAction::CreateSP(this, &SNiagaraStack::CanOpenSourceEmitter)));
 				MenuBuilder.AddMenuEntry(
 					LOCTEXT("ShowEmitterInContentBrowser", "Show in content browser"),
 					LOCTEXT("ShowEmitterInContentBrowserToolTip", "Show the emitter in this stack in the Content Browser"),
@@ -433,7 +433,28 @@ bool SNiagaraStack::IsEntryFocusedInSearch(UNiagaraStackEntry* Entry) const
 
 void SNiagaraStack::OpenSourceEmitter()
 {
-	FAssetEditorManager::Get().OpenEditorForAsset(const_cast<UNiagaraEmitter*>(StackViewModel->GetEmitterHandleViewModel()->GetEmitterHandle()->GetSource()));
+	if (StackViewModel && StackViewModel->GetEmitterHandleViewModel().IsValid() && StackViewModel->GetEmitterHandleViewModel()->GetEmitterHandle())
+	{
+		UNiagaraEmitter* Emitter = const_cast<UNiagaraEmitter*>(StackViewModel->GetEmitterHandleViewModel()->GetEmitterHandle()->GetSource());
+		if (Emitter != nullptr)
+		{
+			FAssetEditorManager::Get().OpenEditorForAsset(Emitter);
+		}
+	}
+}
+
+bool SNiagaraStack::CanOpenSourceEmitter() const
+{
+	if (StackViewModel && StackViewModel->GetEmitterHandleViewModel().IsValid() && StackViewModel->GetEmitterHandleViewModel()->GetEmitterHandle())
+	{
+		UNiagaraEmitter* Emitter = const_cast<UNiagaraEmitter*>(StackViewModel->GetEmitterHandleViewModel()->GetEmitterHandle()->GetSource());
+		if (Emitter != nullptr)
+		{
+			return true;
+		}
+	}
+
+	return false;
 }
 
 void SNiagaraStack::SetEmitterEnabled(bool bIsEnabled)
diff --git a/Engine/Plugins/FX/Niagara/Source/NiagaraEditorWidgets/Private/SNiagaraStack.h b/Engine/Plugins/FX/Niagara/Source/NiagaraEditorWidgets/Private/SNiagaraStack.h
--- a/Engine/Plugins/FX/Niagara/Source/NiagaraEditorWidgets/Private/SNiagaraStack.h
+++ b/Engine/Plugins/FX/Niagara/Source/NiagaraEditorWidgets/Private/SNiagaraStack.h
@@ -89,6 +89,7 @@ private:
 	
 	// Inline menu commands
 	void OpenSourceEmitter();
+	bool CanOpenSourceEmitter() const;
 	void SetEmitterEnabled(bool bIsEnabled);
 	bool CheckEmitterEnabledStatus(bool bIsEnabled);
 	void ShowEmitterInContentBrowser();
