Making Houdini importer only use HCSV files as that doesn't conflict with the other CSV file importers that we support.
Also synchronizing importer to latest version in Dev-Niagara.

#tests all provided example assets load and look correct

diff --git a/Engine/Plugins/FX/HoudiniNiagara/Source/HoudiniNiagara/Classes/NiagaraDataInterfaceHoudiniCSV.h b/Engine/Plugins/FX/HoudiniNiagara/Source/HoudiniNiagara/Classes/NiagaraDataInterfaceHoudiniCSV.h
--- a/Engine/Plugins/FX/HoudiniNiagara/Source/HoudiniNiagara/Classes/NiagaraDataInterfaceHoudiniCSV.h
+++ b/Engine/Plugins/FX/HoudiniNiagara/Source/HoudiniNiagara/Classes/NiagaraDataInterfaceHoudiniCSV.h
@@ -31,6 +31,69 @@
 #include "NiagaraDataInterface.h"
 #include "NiagaraDataInterfaceHoudiniCSV.generated.h"
 
+USTRUCT()
+struct FHoudiniEvent
+{
+	GENERATED_USTRUCT_BODY()
+
+	UPROPERTY(EditAnywhere, Category = "Houdini Event")
+	FVector Position;
+
+	UPROPERTY(EditAnywhere, Category = "Houdini Event")
+	FVector Normal;
+
+	UPROPERTY(EditAnywhere, Category = "Houdini Event")
+	float Impulse;
+
+	UPROPERTY(EditAnywhere, Category = "Houdini Event")
+	FVector Velocity;
+
+	UPROPERTY(EditAnywhere, Category = "Houdini Event")
+	int32 PointID;
+
+	UPROPERTY(EditAnywhere, Category = "Houdini Event")
+	float Time;	
+
+	UPROPERTY(EditAnywhere, Category = "Houdini Event")
+	float Life;
+
+	UPROPERTY(EditAnywhere, Category = "Houdini Event")
+	FLinearColor Color;
+
+	UPROPERTY(EditAnywhere, Category = "Houdini Event")
+	int32 Type;	
+
+	FHoudiniEvent()
+		: Position(FVector::ZeroVector)
+		, Normal(FVector::ZeroVector)
+		, Impulse(0.0f)
+		, Velocity( FVector::ZeroVector)
+		, PointID(0)
+		, Time(0.0f)
+		, Life(0.0f)
+		, Color(FLinearColor::White)
+		, Type(0)		 
+	{
+	}
+
+	inline bool operator==(const FHoudiniEvent& Other) const
+	{
+		if ( (Other.Position != Position) || (Other.Normal != Normal)
+			|| (Other.Velocity != Velocity) || (Other.Impulse != Impulse)
+			|| (Other.PointID != PointID) || (Other.Time != Time)
+			|| (Other.Life != Life) || (Other.Color != Color)
+			|| (Other.Type != Type) )
+			return false;
+
+		return true;
+	}
+
+	inline bool operator!=(const FHoudiniEvent& Other) const
+	{
+		return !(*this == Other);
+	}
+};
+
 
 /** Data Interface allowing sampling of Houdini CSV files. */
 UCLASS(EditInlineNew, Category = "Houdini Niagara", meta = (DisplayName = "Houdini Array Info"))
@@ -102,6 +165,10 @@ public:
 	template<typename RowParamType>
 	void GetColor(FVectorVMContext& Context);
 
+	// Returns the impulse value for a given row in the CSV file
+	template<typename RowParamType>
+	void GetImpulse(FVectorVMContext& Context);
+
 	// Returns the position and time for a given row in the CSV file
 	template<typename RowParamType>
 	void GetPositionAndTime(FVectorVMContext& Context);
diff --git a/Engine/Plugins/FX/HoudiniNiagara/Source/HoudiniNiagara/Private/HoudiniCSV.cpp b/Engine/Plugins/FX/HoudiniNiagara/Source/HoudiniNiagara/Private/HoudiniCSV.cpp
--- a/Engine/Plugins/FX/HoudiniNiagara/Source/HoudiniNiagara/Private/HoudiniCSV.cpp
+++ b/Engine/Plugins/FX/HoudiniNiagara/Source/HoudiniNiagara/Private/HoudiniCSV.cpp
@@ -473,6 +473,11 @@ bool UHoudiniCSV::ParseCSVTitleRow( const FString& TitleRow, const FString& Firs
 			if (!IsValidAttributeColumnIndex(EHoudiniAttributes::TYPE))
 				SpecialAttributesColumnIndexes[EHoudiniAttributes::TYPE] = n;
 		}
+		else if ( CurrentTitle.Equals(TEXT("impulse"), ESearchCase::IgnoreCase ) )
+		{
+			if (!IsValidAttributeColumnIndex(EHoudiniAttributes::IMPULSE))
+				SpecialAttributesColumnIndexes[EHoudiniAttributes::IMPULSE] = n;
+		}
     }
 
 	// Read the first row of the CSV file, and look for packed vectors value (X,Y,Z)
@@ -524,14 +529,9 @@ bool UHoudiniCSV::ParseCSVTitleRow( const FString& TitleRow, const FString& Firs
 		// Extract the special attributes column indexes we found
 		int32 PositionColumnIndex = GetAttributeColumnIndex(EHoudiniAttributes::POSITION);
 		int32 NormalColumnIndex = GetAttributeColumnIndex(EHoudiniAttributes::NORMAL);
-		int32 TimeColumnIndex = GetAttributeColumnIndex(EHoudiniAttributes::TIME);
-		int32 IDColumnIndex = GetAttributeColumnIndex(EHoudiniAttributes::POINTID);
-		int32 AliveColumnIndex = GetAttributeColumnIndex(EHoudiniAttributes::ALIVE);
-		int32 LifeColumnIndex = GetAttributeColumnIndex(EHoudiniAttributes::LIFE);
 		int32 ColorColumnIndex = GetAttributeColumnIndex(EHoudiniAttributes::COLOR);
 		int32 AlphaColumnIndex = GetAttributeColumnIndex(EHoudiniAttributes::ALPHA);
-		int32 VelocityColumnIndex = GetAttributeColumnIndex(EHoudiniAttributes::VELOCITY);
-		int32 TypeColumnIndex = GetAttributeColumnIndex(EHoudiniAttributes::TYPE);
+		int32 VelocityColumnIndex = GetAttributeColumnIndex(EHoudiniAttributes::VELOCITY);		
 
 		// Expand TitleRowArray
 		if ( ( FoundPackedVectorColumnIndex == PositionColumnIndex ) && ( FoundVectorSize == 3 ) )
@@ -763,6 +763,18 @@ bool UHoudiniCSV::GetVelocityValue( const int32& rowIndex, FVector& value ) cons
 	return true;
 }
 
+// Returns an impulse value for a given point in the CSV file
+bool UHoudiniCSV::GetImpulseValue(const int32& rowIndex, float& value) const
+{
+	float temp;
+	if (!GetFloatValue(rowIndex, GetAttributeColumnIndex(EHoudiniAttributes::IMPULSE), temp))
+		return false;
+
+	value = temp;
+
+	return true;
+}
+
 // Returns the number of points found in the CSV file
 int32 UHoudiniCSV::GetNumberOfPoints() const
 {
@@ -1244,6 +1256,9 @@ UHoudiniCSV::GetAssetRegistryTags(TArray< FAssetRegistryTag > & OutTags) const
 	if (IsValidAttributeColumnIndex(EHoudiniAttributes::NORMAL))
 		SpecialAttr += TEXT("Normal ");
 
+	if (IsValidAttributeColumnIndex(EHoudiniAttributes::IMPULSE))
+		SpecialAttr += TEXT("Impulse ");
+
 	if (IsValidAttributeColumnIndex(EHoudiniAttributes::VELOCITY))
 		SpecialAttr += TEXT("Velocity ");
 
diff --git a/Engine/Plugins/FX/HoudiniNiagara/Source/HoudiniNiagara/Private/NiagaraDataInterfaceHoudiniCSV.cpp b/Engine/Plugins/FX/HoudiniNiagara/Source/HoudiniNiagara/Private/NiagaraDataInterfaceHoudiniCSV.cpp
--- a/Engine/Plugins/FX/HoudiniNiagara/Source/HoudiniNiagara/Private/NiagaraDataInterfaceHoudiniCSV.cpp
+++ b/Engine/Plugins/FX/HoudiniNiagara/Source/HoudiniNiagara/Private/NiagaraDataInterfaceHoudiniCSV.cpp
@@ -55,6 +55,7 @@ static const FName GetNormalName("GetNormal");
 static const FName GetTimeName("GetTime");
 static const FName GetVelocityName("GetVelocity");
 static const FName GetColorName("GetColor");
+static const FName GetImpulseName("GetImpulse");
 static const FName GetPositionAndTimeName("GetPositionAndTime");
 
 static const FName GetNumberOfPointsName("GetNumberOfPoints");
@@ -89,6 +90,7 @@ void UNiagaraDataInterfaceHoudiniCSV::PostInitProperties()
     if (HasAnyFlags(RF_ClassDefaultObject))
     {
 	    FNiagaraTypeRegistry::Register(FNiagaraTypeDefinition(GetClass()), true, false, false);
+		FNiagaraTypeRegistry::Register(FHoudiniEvent::StaticStruct(), true, true, false);
     }
 
 	FloatValuesGPUBufferDirty = true;
@@ -354,6 +356,22 @@ void UNiagaraDataInterfaceHoudiniCSV::GetFunctions(TArray<FNiagaraFunctionSignat
 		OutFunctions.Add(Sig);
 	}
 
+	{
+		// GetImpulse
+		FNiagaraFunctionSignature Sig;
+		Sig.Name = GetImpulseName;
+		Sig.bMemberFunction = true;
+		Sig.bRequiresContext = false;
+		Sig.Inputs.Add(FNiagaraVariable(FNiagaraTypeDefinition(GetClass()), TEXT("CSV")));			// CSV in
+		Sig.Inputs.Add(FNiagaraVariable(FNiagaraTypeDefinition::GetIntDef(), TEXT("Row")));			// Row Index In
+		Sig.Outputs.Add(FNiagaraVariable(FNiagaraTypeDefinition::GetFloatDef(), TEXT("Impulse")));	// Float Out
+
+		Sig.SetDescription(LOCTEXT("DataInterfaceHoudini_GetImpulse",
+			"Helper function returning the impulse value for a given Row in the CSV file.\n"));
+
+		OutFunctions.Add(Sig);
+	}
+
     {
 		// GetNumberOfPoints
 		FNiagaraFunctionSignature Sig;
@@ -586,6 +604,7 @@ DEFINE_NDI_RAW_FUNC_BINDER(UNiagaraDataInterfaceHoudiniCSV, GetNormal);
 DEFINE_NDI_RAW_FUNC_BINDER(UNiagaraDataInterfaceHoudiniCSV, GetTime);
 DEFINE_NDI_RAW_FUNC_BINDER(UNiagaraDataInterfaceHoudiniCSV, GetColor);
 DEFINE_NDI_RAW_FUNC_BINDER(UNiagaraDataInterfaceHoudiniCSV, GetVelocity);
+DEFINE_NDI_RAW_FUNC_BINDER(UNiagaraDataInterfaceHoudiniCSV, GetImpulse);
 DEFINE_NDI_RAW_FUNC_BINDER(UNiagaraDataInterfaceHoudiniCSV, GetPositionAndTime);
 DEFINE_NDI_RAW_FUNC_BINDER(UNiagaraDataInterfaceHoudiniCSV, GetLastRowIndexAtTime);
 DEFINE_NDI_RAW_FUNC_BINDER(UNiagaraDataInterfaceHoudiniCSV, GetPointIDsToSpawnAtTime);
@@ -635,7 +654,11 @@ void UNiagaraDataInterfaceHoudiniCSV::GetVMExternalFunction(const FVMExternalFun
 	{
 		TNDIParamBinder<0, int32, NDI_RAW_FUNC_BINDER(UNiagaraDataInterfaceHoudiniCSV, GetColor)>::Bind(this, BindingInfo, InstanceData, OutFunc);
 	}
-    else if (BindingInfo.Name == GetPositionAndTimeName && BindingInfo.GetNumInputs() == 1 && BindingInfo.GetNumOutputs() == 4)
+	else if (BindingInfo.Name == GetImpulseName && BindingInfo.GetNumInputs() == 1 && BindingInfo.GetNumOutputs() == 1)
+	{
+		TNDIParamBinder<0, int32, NDI_RAW_FUNC_BINDER(UNiagaraDataInterfaceHoudiniCSV, GetImpulse)>::Bind(this, BindingInfo, InstanceData, OutFunc);
+	}
+	else if (BindingInfo.Name == GetPositionAndTimeName && BindingInfo.GetNumInputs() == 1 && BindingInfo.GetNumOutputs() == 4)
     {
 		TNDIParamBinder<0, int32, NDI_RAW_FUNC_BINDER(UNiagaraDataInterfaceHoudiniCSV, GetPositionAndTime)>::Bind(this, BindingInfo, InstanceData, OutFunc);
     }
@@ -947,6 +970,26 @@ void UNiagaraDataInterfaceHoudiniCSV::GetColor(FVectorVMContext& Context)
 	}
 }
 
+template<typename RowParamType>
+void UNiagaraDataInterfaceHoudiniCSV::GetImpulse(FVectorVMContext& Context)
+{
+	RowParamType RowParam(Context);
+	FRegisterHandler<float> OutValue(Context);
+
+	for (int32 i = 0; i < Context.NumInstances; ++i)
+	{
+		int32 row = RowParam.Get();
+
+		float value = 0.0f;
+		if (HoudiniCSVAsset)
+			HoudiniCSVAsset->GetImpulseValue(row, value);
+
+		*OutValue.GetDest() = value;
+		RowParam.Advance();
+		OutValue.Advance();
+	}
+}
+
 // Returns the last index of the points that should be spawned at time t
 template<typename TimeParamType>
 void UNiagaraDataInterfaceHoudiniCSV::GetLastRowIndexAtTime(FVectorVMContext& Context)
@@ -1528,6 +1571,17 @@ bool UNiagaraDataInterfaceHoudiniCSV::GetFunctionHLSL(const FName& DefinitionFun
 		OutHLSL += TEXT("\n}\n");
 		return true;
 	}
+	else if (DefinitionFunctionName == GetImpulseName)
+	{
+		// GetImpulse(int In_Row, out float Out_Value)
+		OutHLSL += TEXT("void ") + InstanceFunctionName + TEXT("(int In_Row, out float Out_Impulse) \n{\n");
+
+		OutHLSL += TEXT("\tint In_ImpulseCol = ") + GetSpecAttributeColumnIndex(EHoudiniAttributes::IMPULSE) + TEXT(";\n");
+		OutHLSL += ReadFloatInBuffer(TEXT("Out_Impulse"), TEXT("In_Row"), TEXT("In_ImpulseCol"));
+
+		OutHLSL += TEXT("\n}\n");
+		return true;
+	}
 	else if (DefinitionFunctionName == GetColorName)
 	{
 		// GetColor(int In_Row, out float4 Out_Value)
@@ -1618,9 +1672,11 @@ bool UNiagaraDataInterfaceHoudiniCSV::GetFunctionHLSL(const FName& DefinitionFun
 			// Check if we need to reset lastspawnedPointID
 			OutHLSL += TEXT("\tif ( last_id < ") + LastSpawnedPointIDVar + TEXT(" || In_Time <= ") + LastSpawnTimeVar + TEXT(")\n");
 				OutHLSL += TEXT("\t\t{") + LastSpawnedPointIDVar + TEXT(" = -1; }\n");
+
 			// Nothing to spawn, In_Time is lower than the point's time
 			OutHLSL += TEXT("\tif ( last_id < 0 )\n");
 				OutHLSL += TEXT("\t\t{") + LastSpawnedPointIDVar + TEXT(" = -1;Out_Count = 0; Out_MinID = 0; Out_MaxID = 0;return;}\n");
+
 			// We dont have any new point to spawn
 			OutHLSL += TEXT("\tif ( last_id == ") + LastSpawnedPointIDVar + TEXT(")\n");
 				OutHLSL += TEXT("{ Out_MinID = last_id; Out_MaxID = last_id; Out_Count = 0; return;}\n");
diff --git a/Engine/Plugins/FX/HoudiniNiagara/Source/HoudiniNiagara/Public/HoudiniCSV.h b/Engine/Plugins/FX/HoudiniNiagara/Source/HoudiniNiagara/Public/HoudiniCSV.h
--- a/Engine/Plugins/FX/HoudiniNiagara/Source/HoudiniNiagara/Public/HoudiniCSV.h
+++ b/Engine/Plugins/FX/HoudiniNiagara/Source/HoudiniNiagara/Public/HoudiniCSV.h
@@ -44,6 +44,7 @@ enum EHoudiniAttributes
 	ALPHA,
 	VELOCITY,
 	TYPE,
+	IMPULSE,
 
 	HOUDINI_ATTR_SIZE,
 	HOUDINI_ATTR_END = HOUDINI_ATTR_SIZE - 1
@@ -132,6 +133,9 @@ class HOUDININIAGARA_API UHoudiniCSV : public UObject
 	// Returns a Velocity Vector3 for a given point in the CSV file
 	UFUNCTION(BlueprintCallable, Category = "Houdini Attributes Data")
 	bool GetVelocityValue(const int32& rowIndex, FVector& value ) const;
+	// Returns an Impulse float value for a given point in the CSV file
+	UFUNCTION(BlueprintCallable, Category = "Houdini Attributes Data")
+	bool GetImpulseValue(const int32& rowIndex, float& value) const;
 
 	// Get the last row index for a given time value (the row with a time smaller or equal to desiredTime)
 	// If the CSV file doesn't have time informations, returns false and set LastRowIndex to the last row in the file
diff --git a/Engine/Plugins/FX/HoudiniNiagara/Source/HoudiniNiagaraEditor/Private/HoudiniCSVFactory.cpp b/Engine/Plugins/FX/HoudiniNiagara/Source/HoudiniNiagaraEditor/Private/HoudiniCSVFactory.cpp
--- a/Engine/Plugins/FX/HoudiniNiagara/Source/HoudiniNiagaraEditor/Private/HoudiniCSVFactory.cpp
+++ b/Engine/Plugins/FX/HoudiniNiagara/Source/HoudiniNiagaraEditor/Private/HoudiniCSVFactory.cpp
@@ -53,7 +53,6 @@ UHoudiniCSVFactory::UHoudiniCSVFactory(const FObjectInitializer& ObjectInitializ
     bText = true;
 
     // Add supported formats.
-    Formats.Add(FString(TEXT("csv;")) + NSLOCTEXT("HoudiniCSVFactory", "FormatCSV", "CSV File").ToString());
     Formats.Add(FString(TEXT("hcsv;")) + NSLOCTEXT("HoudiniCSVFactory", "FormatHCSV", "HCSV File").ToString());
 }
 
@@ -110,7 +109,7 @@ UHoudiniCSVFactory::FactoryCanImport(const FString& Filename)
 {
     const FString Extension = FPaths::GetExtension(Filename);
 
-    if ( Extension == TEXT("csv") || Extension == TEXT("hcsv") )
+    if (Extension == TEXT("hcsv") )
     {
 		return true;
     }
