Niagara - Houdini:

Houdini Niagara Data Interface:
- Removed the GetCSVFloatValueByString function as String aren't currently supported by Niagara.
- Particles in the CSV file can now be updated over time (not just spawned)
- Added GetParticlePositionAtTime, GetParticleValueAtTime, GetParticleVectorValueAtTime returning a linearly interpolated value for a given particle at a given time.
- Added GetParticleIndexesAtTime for getting the previous/next row indexes and weight to access the values for a given particle at a given time and handle the interpolation of the values.
- Added GetCSVVectorValue for accessing a Vector value at a given row/col.

Houdini CSV Assets now looks for the following attributes in the CSV "Title" line:
- pos for position.
- id and # for particle ID.
- alive and life for calculating a particles LifeTime.

diff --git a/Engine/Plugins/FX/HoudiniNiagara/Source/HoudiniNiagara/Classes/NiagaraDataInterfaceHoudiniCSV.h b/Engine/Plugins/FX/HoudiniNiagara/Source/HoudiniNiagara/Classes/NiagaraDataInterfaceHoudiniCSV.h
--- a/Engine/Plugins/FX/HoudiniNiagara/Source/HoudiniNiagara/Classes/NiagaraDataInterfaceHoudiniCSV.h
+++ b/Engine/Plugins/FX/HoudiniNiagara/Source/HoudiniNiagara/Classes/NiagaraDataInterfaceHoudiniCSV.h
@@ -1,5 +1,5 @@
 /*
-* Copyright (c) <2017> Side Effects Software Inc.
+* Copyright (c) <2018> Side Effects Software Inc.
 *
 * Permission is hereby granted, free of charge, to any person obtaining a copy
 * of this software and associated documentation files (the "Software"), to deal
@@ -39,14 +39,9 @@ class HOUDININIAGARA_API UNiagaraDataInterfaceHoudiniCSV : public UNiagaraDataIn
 	GENERATED_UCLASS_BODY()
 public:
 
-	//UPROPERTY( EditAnywhere, Category = "Houdini Niagara", meta = (DisplayName = "Houdini CSV File Path" ) )
-	//FFilePath CSVFileName;
-
 	UPROPERTY( EditAnywhere, Category = "Houdini Niagara", meta = (DisplayName = "Houdini CSV File" ) )
 	UHoudiniCSV* CSVFile;
 
-	//void UpdateDataFromCSVFile();
-
 	//----------------------------------------------------------------------------
 	//UObject Interface
 	virtual void PostInitProperties() override;
@@ -70,8 +65,10 @@ public:
 	template<typename RowParamType, typename ColParamType>
 	void GetCSVFloatValue(FVectorVMContext& Context);
 
+	/*
 	template<typename RowParamType, typename ColTitleParamType>
 	void GetCSVFloatValueByString(FVectorVMContext& Context);
+	*/
 
 	// Returns a Vector3 value for a given point in the csv file
 	template<typename RowParamType, typename ColParamType>
@@ -103,6 +100,22 @@ public:
 	// Returns the indexes (min, max) and number of particles that should be spawned at time t
 	template<typename TimeParamType>
 	void GetParticleIndexesToSpawnAtTime(FVectorVMContext& Context);
+
+	// Returns the position for a given particle at a given time
+	template<typename NParamType, typename TimeParamType>
+	void GetParticlePositionAtTime(FVectorVMContext& Context);
+
+	// Returns a float value for a given particle at a given time
+	template<typename NParamType, typename ColParamType, typename TimeParamType>
+	void GetParticleValueAtTime(FVectorVMContext& Context);
+
+	// Returns a Vector value for a given particle at a given time
+	template<typename NParamType, typename ColParamType, typename TimeParamType>
+	void GetParticleVectorValueAtTime(FVectorVMContext& Context);
+
+	// Returns the line indexes (previous, next) for reading values for a given particle at a given time
+	template<typename NParamType, typename TimeParamType>
+	void GetParticleIndexesAtTime(FVectorVMContext& Context);
 	
 	//----------------------------------------------------------------------------
 	// GPU / HLSL Functions
diff --git a/Engine/Plugins/FX/HoudiniNiagara/Source/HoudiniNiagara/HoudiniNiagara.Build.cs b/Engine/Plugins/FX/HoudiniNiagara/Source/HoudiniNiagara/HoudiniNiagara.Build.cs
--- a/Engine/Plugins/FX/HoudiniNiagara/Source/HoudiniNiagara/HoudiniNiagara.Build.cs
+++ b/Engine/Plugins/FX/HoudiniNiagara/Source/HoudiniNiagara/HoudiniNiagara.Build.cs
@@ -1,5 +1,5 @@
 /*
-* Copyright (c) <2017> Side Effects Software Inc.
+* Copyright (c) <2018> Side Effects Software Inc.
 *
 * Permission is hereby granted, free of charge, to any person obtaining a copy
 * of this software and associated documentation files (the "Software"), to deal
@@ -39,7 +39,7 @@ public class HoudiniNiagara : ModuleRules
 				"Engine",
 				"Slate",
 				"SlateCore",
-              			"NiagaraCore",
+				"NiagaraCore",
 				"Niagara",
 				"NiagaraShader",
 				"RenderCore",
@@ -53,7 +53,7 @@ public class HoudiniNiagara : ModuleRules
 			new string[]
 			{
 				"Core",
-              			"NiagaraCore",
+              	"NiagaraCore",
 				"Niagara",
 				"NiagaraShader",
 				"CoreUObject",
diff --git a/Engine/Plugins/FX/HoudiniNiagara/Source/HoudiniNiagara/Private/HoudiniCSV.cpp b/Engine/Plugins/FX/HoudiniNiagara/Source/HoudiniNiagara/Private/HoudiniCSV.cpp
--- a/Engine/Plugins/FX/HoudiniNiagara/Source/HoudiniNiagara/Private/HoudiniCSV.cpp
+++ b/Engine/Plugins/FX/HoudiniNiagara/Source/HoudiniNiagara/Private/HoudiniCSV.cpp
@@ -1,5 +1,5 @@
 /*
-* Copyright (c) <2017> Side Effects Software Inc.
+* Copyright (c) <2018> Side Effects Software Inc.
 *
 * Permission is hereby granted, free of charge, to any person obtaining a copy
 * of this software and associated documentation files (the "Software"), to deal
@@ -36,9 +36,15 @@ DEFINE_LOG_CATEGORY(LogHoudiniNiagara);
 UHoudiniCSV::UHoudiniCSV( const FObjectInitializer& ObjectInitializer )
 	: Super( ObjectInitializer ),
 	NumberOfLines( -1 ),
-	NumberOfColumns( -1 )
+	NumberOfColumns( -1 ),
+	NumberOfParticles( -1 )
 {
-        
+	PositionColumnIndex = INDEX_NONE;
+	NormalColumnIndex = INDEX_NONE;
+	TimeColumnIndex = INDEX_NONE;
+	IDColumnIndex = INDEX_NONE;
+	AliveColumnIndex = INDEX_NONE;
+	LifeColumnIndex = INDEX_NONE;
 }
 
 void UHoudiniCSV::SetFileName( const FString& TheFileName )
@@ -49,19 +55,19 @@ void UHoudiniCSV::SetFileName( const FString& TheFileName )
 bool UHoudiniCSV::UpdateFromFile( const FString& TheFileName )
 {
     if ( TheFileName.IsEmpty() )
-	return false;
+		return false;
 
     // Check the CSV file exists
     const FString FullCSVFilename = FPaths::ConvertRelativePathToFull( TheFileName );
     if ( !FPaths::FileExists( FullCSVFilename ) )
-	return false;
+		return false;
 
     FileName = TheFileName;
 
     // Parse the file to a string array
     TArray<FString> StringArray;
     if ( !FFileHelper::LoadFileToStringArray( StringArray, *FullCSVFilename ) )
-	return false;
+		return false;
 
     return UpdateFromStringArray( StringArray );
 }
@@ -71,16 +77,20 @@ bool UHoudiniCSV::UpdateFromStringArray( TArray<FString>& RawStringArray )
     // Reset the CSV sizes
     NumberOfColumns = 0;
     NumberOfLines = 0;
+	NumberOfParticles = 0;
 
     // Reset the position, normal and time indexes
     PositionColumnIndex = INDEX_NONE;
     NormalColumnIndex = INDEX_NONE;
     TimeColumnIndex = INDEX_NONE;
+	IDColumnIndex = INDEX_NONE;
+	AliveColumnIndex = INDEX_NONE;
+	LifeColumnIndex = INDEX_NONE;
 
     if ( RawStringArray.Num() <= 0 )
     {
-	UE_LOG( LogHoudiniNiagara, Error, TEXT( "Could not load the CSV file, error: not enough lines." ) );
-	return false;
+		UE_LOG( LogHoudiniNiagara, Error, TEXT( "Could not load the CSV file, error: not enough lines." ) );
+		return false;
     }
 
     // Remove empty lines from the CSV
@@ -90,8 +100,8 @@ bool UHoudiniCSV::UpdateFromStringArray( TArray<FString>& RawStringArray )
     NumberOfLines = RawStringArray.Num() - 1;
     if ( NumberOfLines < 1 )
     {
-	UE_LOG( LogHoudiniNiagara, Error, TEXT( "Could not load the CSV file, error: not enough lines." ) );
-	return false;
+		UE_LOG( LogHoudiniNiagara, Error, TEXT( "Could not load the CSV file, error: not enough lines." ) );
+		return false;
     }
 
     // Get the number of values per lines via the title line
@@ -101,32 +111,53 @@ bool UHoudiniCSV::UpdateFromStringArray( TArray<FString>& RawStringArray )
     NumberOfColumns = TitleRowArray.Num();
     if ( NumberOfColumns < 1 )
     {
-	UE_LOG( LogHoudiniNiagara, Error, TEXT( "Could not load the CSV file, error: not enough columns." ) );
-	return false;
+		UE_LOG( LogHoudiniNiagara, Error, TEXT( "Could not load the CSV file, error: not enough columns." ) );
+		return false;
     }
 
     // Look for the position, normal and time attributes indexes
     for ( int32 n = 0; n < TitleRowArray.Num(); n++ )
     {
-	if ( TitleRowArray[ n ].Equals( TEXT("P"), ESearchCase::IgnoreCase )
-	    || TitleRowArray[ n ].Equals( TEXT("Px"), ESearchCase::IgnoreCase ) 
-	    || TitleRowArray[ n ].Equals( TEXT("X"), ESearchCase::IgnoreCase ) )
-	{
-	    if ( PositionColumnIndex == INDEX_NONE )
-		PositionColumnIndex = n;
-	}
-	else if ( TitleRowArray[ n ].Equals( TEXT("N"), ESearchCase::IgnoreCase ) 
-	    || TitleRowArray[ n ].Equals( TEXT("Nx"), ESearchCase::IgnoreCase ) )
-	{
-	    if ( NormalColumnIndex == INDEX_NONE )
-		NormalColumnIndex = n;
-	}
-	else if ( ( TitleRowArray[ n ].Equals( TEXT("T"), ESearchCase::IgnoreCase ) )
-	    || ( TitleRowArray[ n ].Contains( TEXT("time"), ESearchCase::IgnoreCase ) ) )
-	{
-	    if ( TimeColumnIndex == INDEX_NONE )
-		TimeColumnIndex = n;
-	}
+		FString CurrentTitle = TitleRowArray[ n ];
+		CurrentTitle.ToLowerInline();
+		CurrentTitle.ReplaceInline( TEXT(" "), TEXT("") );
+
+		if ( CurrentTitle.Equals( TEXT("P"), ESearchCase::IgnoreCase )
+			|| CurrentTitle.Equals( TEXT("Px"), ESearchCase::IgnoreCase )
+			|| CurrentTitle.Equals( TEXT("X"), ESearchCase::IgnoreCase )
+			|| CurrentTitle.Equals(TEXT("pos"), ESearchCase::IgnoreCase ) )
+		{
+			if ( PositionColumnIndex == INDEX_NONE )
+				PositionColumnIndex = n;
+		}
+		else if ( CurrentTitle.Equals( TEXT("N"), ESearchCase::IgnoreCase )
+			|| CurrentTitle.Equals( TEXT("Nx"), ESearchCase::IgnoreCase ) )
+		{
+			if ( NormalColumnIndex == INDEX_NONE )
+				NormalColumnIndex = n;
+		}
+		else if ( ( CurrentTitle.Equals( TEXT("T"), ESearchCase::IgnoreCase ) )
+			|| ( CurrentTitle.Contains( TEXT("time"), ESearchCase::IgnoreCase ) ) )
+		{
+			if ( TimeColumnIndex == INDEX_NONE )
+				TimeColumnIndex = n;
+		}
+		else if ( ( CurrentTitle.Equals( TEXT("#"), ESearchCase::IgnoreCase ) )
+			|| ( CurrentTitle.Equals( TEXT("id"), ESearchCase::IgnoreCase ) ) )
+		{
+			if ( IDColumnIndex == INDEX_NONE )
+				IDColumnIndex = n;
+		}
+		else if ( CurrentTitle.Equals( TEXT("alive"), ESearchCase::IgnoreCase ) )
+		{
+			if (AliveColumnIndex == INDEX_NONE)
+				AliveColumnIndex = n;
+		}
+		else if ( CurrentTitle.Equals( TEXT("life"), ESearchCase::IgnoreCase ) )
+		{
+			if ( LifeColumnIndex == INDEX_NONE)
+				LifeColumnIndex = n;
+		}
     }
 
     // Read the first line of the CSV file, and look for packed vectors value (X,Y,Z)
@@ -136,100 +167,109 @@ bool UHoudiniCSV::UpdateFromStringArray( TArray<FString>& RawStringArray )
     bool HasPackedVectors = false;
     while ( FoundPackedVectorCharIndex != INDEX_NONE )
     {
-	// Try to find ( in the line
-	FoundPackedVectorCharIndex = FirstLine.Find( TEXT("("), ESearchCase::IgnoreCase, ESearchDir::FromStart, FoundPackedVectorCharIndex );
-	if ( FoundPackedVectorCharIndex == INDEX_NONE )
-	    break;
-	// We want to know which column this char belong to
-	int32 FoundPackedVectorColumnIndex = INDEX_NONE;
-	{
-	    // Chop the first line up to the found character
-	    FString FirstLineLeft = FirstLine.Left( FoundPackedVectorCharIndex );
-
-	    // ReplaceInLine returns the number of occurences of ",", that's what we want! 
-	    FoundPackedVectorColumnIndex = FirstLineLeft.ReplaceInline(TEXT(","), TEXT(""));
-	}
-
-	if ( !TitleRowArray.IsValidIndex( FoundPackedVectorColumnIndex ) )
-	{
-	    UE_LOG( LogHoudiniNiagara, Warning,
-		TEXT( "Error while parsing the CSV File. Couldn't unpack vector found at character %d in the first line!" ),
-		FoundPackedVectorCharIndex + 1 );
-	    continue;
-	}
-	    
-
-	// We found a packed vector, get its sizw
-	int32 FoundVectorSize = 0;
-	{
-	    // Extract the vector string
-	    int32 FoundPackedVectorEndCharIndex = FirstLine.Find( TEXT(")"), ESearchCase::IgnoreCase, ESearchDir::FromStart, FoundPackedVectorCharIndex );
-	    FString VectorString = FirstLine.Mid( FoundPackedVectorCharIndex + 1, FoundPackedVectorEndCharIndex - FoundPackedVectorCharIndex - 1 );
-
-	    // Use ReplaceInLine to count the number of , to get the vector's size!
-	    FoundVectorSize = VectorString.ReplaceInline( TEXT(","), TEXT("") ) + 1;
-	}
-
-	if ( FoundVectorSize < 2 )
-	    continue;
-
-	// Increment the number of columns
-	NumberOfColumns += ( FoundVectorSize - 1 );
-
-	// Expand TitleRowArray
-	if ( ( FoundPackedVectorColumnIndex == PositionColumnIndex ) && ( FoundVectorSize == 3 ) )
-	{
-	    // Expand P to PX,Py,Pz
-	    TitleRowArray[ PositionColumnIndex ] = TEXT("Px");
-	    TitleRowArray.Insert( TEXT( "Py" ), PositionColumnIndex + 1 );
-	    TitleRowArray.Insert( TEXT( "Pz" ), PositionColumnIndex + 2 );
-	}
-	else if ( ( FoundPackedVectorColumnIndex == NormalColumnIndex ) && ( FoundVectorSize == 3 ) )
-	{
-	    // Expand N to Nx,Ny,Nz
-	    TitleRowArray[ NormalColumnIndex ] = TEXT("Nx");
-	    TitleRowArray.Insert( TEXT("Ny"), NormalColumnIndex + 1 );
-	    TitleRowArray.Insert( TEXT("Nz"), NormalColumnIndex + 2 );
-	}
-	else
-	{
-	    // Expand the vector's title
-	    FString FoundPackedVectortTitle = TitleRowArray[ FoundPackedVectorColumnIndex ];
-	    for ( int32 n = 1; n < FoundVectorSize; n++ )
-	    {
-		FString CurrentTitle = FoundPackedVectortTitle + FString::FromInt( n );
-		TitleRowArray.Insert( CurrentTitle, FoundPackedVectorColumnIndex + n );
-	    }
-	}
-
-	// Eventually offset the stored index
-	if ( PositionColumnIndex != INDEX_NONE && ( PositionColumnIndex > FoundPackedVectorColumnIndex) )
-	    PositionColumnIndex += FoundVectorSize - 1;
-
-	if ( NormalColumnIndex != INDEX_NONE && ( NormalColumnIndex > FoundPackedVectorColumnIndex) )
-	    NormalColumnIndex += FoundVectorSize - 1;
-
-	if ( TimeColumnIndex != INDEX_NONE && ( TimeColumnIndex > FoundPackedVectorColumnIndex) )
-	    TimeColumnIndex += FoundVectorSize - 1;
-
-	HasPackedVectors = true;
-	FoundPackedVectorCharIndex++;
+		// Try to find ( in the line
+		FoundPackedVectorCharIndex = FirstLine.Find( TEXT("("), ESearchCase::IgnoreCase, ESearchDir::FromStart, FoundPackedVectorCharIndex );
+		if ( FoundPackedVectorCharIndex == INDEX_NONE )
+			break;
+
+		// We want to know which column this char belong to
+		int32 FoundPackedVectorColumnIndex = INDEX_NONE;
+		{
+			// Chop the first line up to the found character
+			FString FirstLineLeft = FirstLine.Left( FoundPackedVectorCharIndex );
+
+			// ReplaceInLine returns the number of occurences of ",", that's what we want! 
+			FoundPackedVectorColumnIndex = FirstLineLeft.ReplaceInline(TEXT(","), TEXT(""));
+		}
+
+		if ( !TitleRowArray.IsValidIndex( FoundPackedVectorColumnIndex ) )
+		{
+			UE_LOG( LogHoudiniNiagara, Warning,
+			TEXT( "Error while parsing the CSV File. Couldn't unpack vector found at character %d in the first line!" ),
+			FoundPackedVectorCharIndex + 1 );
+			continue;
+		}
+
+		// We found a packed vector, get its size
+		int32 FoundVectorSize = 0;
+		{
+			// Extract the vector string
+			int32 FoundPackedVectorEndCharIndex = FirstLine.Find( TEXT(")"), ESearchCase::IgnoreCase, ESearchDir::FromStart, FoundPackedVectorCharIndex );
+			FString VectorString = FirstLine.Mid( FoundPackedVectorCharIndex + 1, FoundPackedVectorEndCharIndex - FoundPackedVectorCharIndex - 1 );
+
+			// Use ReplaceInLine to count the number of , to get the vector's size!
+			FoundVectorSize = VectorString.ReplaceInline( TEXT(","), TEXT("") ) + 1;
+		}
+
+		if ( FoundVectorSize < 2 )
+			continue;
+
+		// Increment the number of columns
+		NumberOfColumns += ( FoundVectorSize - 1 );
+
+		// Expand TitleRowArray
+		if ( ( FoundPackedVectorColumnIndex == PositionColumnIndex ) && ( FoundVectorSize == 3 ) )
+		{
+			// Expand P to Px,Py,Pz
+			TitleRowArray[ PositionColumnIndex ] = TEXT("Px");
+			TitleRowArray.Insert( TEXT( "Py" ), PositionColumnIndex + 1 );
+			TitleRowArray.Insert( TEXT( "Pz" ), PositionColumnIndex + 2 );
+		}
+		else if ( ( FoundPackedVectorColumnIndex == NormalColumnIndex ) && ( FoundVectorSize == 3 ) )
+		{
+			// Expand N to Nx,Ny,Nz
+			TitleRowArray[ NormalColumnIndex ] = TEXT("Nx");
+			TitleRowArray.Insert( TEXT("Ny"), NormalColumnIndex + 1 );
+			TitleRowArray.Insert( TEXT("Nz"), NormalColumnIndex + 2 );
+		}
+		else
+		{
+			// Expand the vector's title from V to V, V1, V2, V3 ...
+			FString FoundPackedVectortTitle = TitleRowArray[ FoundPackedVectorColumnIndex ];
+			for ( int32 n = 1; n < FoundVectorSize; n++ )
+			{
+				FString CurrentTitle = FoundPackedVectortTitle + FString::FromInt( n );
+				TitleRowArray.Insert( CurrentTitle, FoundPackedVectorColumnIndex + n );
+			}
+		}
+
+		// Eventually offset the stored index
+		if ( PositionColumnIndex != INDEX_NONE && ( PositionColumnIndex > FoundPackedVectorColumnIndex) )
+			PositionColumnIndex += FoundVectorSize - 1;
+
+		if ( NormalColumnIndex != INDEX_NONE && ( NormalColumnIndex > FoundPackedVectorColumnIndex) )
+			NormalColumnIndex += FoundVectorSize - 1;
+
+		if ( TimeColumnIndex != INDEX_NONE && ( TimeColumnIndex > FoundPackedVectorColumnIndex) )
+			TimeColumnIndex += FoundVectorSize - 1;
+
+		if ( IDColumnIndex != INDEX_NONE && ( IDColumnIndex > FoundPackedVectorColumnIndex ) )
+			IDColumnIndex += FoundVectorSize - 1;
+
+		if ( AliveColumnIndex != INDEX_NONE && ( AliveColumnIndex > FoundPackedVectorColumnIndex ) )
+			AliveColumnIndex += FoundVectorSize - 1;
+
+		if ( LifeColumnIndex != INDEX_NONE && ( LifeColumnIndex > FoundPackedVectorColumnIndex ) )
+			LifeColumnIndex += FoundVectorSize - 1;
+
+		HasPackedVectors = true;
+		FoundPackedVectorCharIndex++;
     }
 
     // For sanity, Check that the number of columns matches the title row and the first line
     {
-	// Check the title row
-	if ( NumberOfColumns != TitleRowArray.Num() )
-	    UE_LOG( LogHoudiniNiagara, Error,
-		TEXT( "Error while parsing the CSV File. Found %d columns but the Title string has %d values! Some values will have an offset!" ),
-		NumberOfColumns, TitleRowArray.Num() );
-
-	// Use ReplaceInLine to count the number of columns in the first line and make sure it's correct
-	int32 FirstLineColumns = FirstLine.ReplaceInline( TEXT(","), TEXT("") ) + 1;
-	if ( NumberOfColumns != FirstLineColumns )
-	    UE_LOG( LogHoudiniNiagara, Error,
-		TEXT("Error while parsing the CSV File. Found %d columns but found %d values in the first line! Some values will have an offset!" ),
-		NumberOfColumns, FirstLineColumns );
+		// Check the title row
+		if ( NumberOfColumns != TitleRowArray.Num() )
+			UE_LOG( LogHoudiniNiagara, Error,
+			TEXT( "Error while parsing the CSV File. Found %d columns but the Title string has %d values! Some values will have an offset!" ),
+			NumberOfColumns, TitleRowArray.Num() );
+
+		// Use ReplaceInLine to count the number of columns in the first line and make sure it's correct
+		int32 FirstLineColumns = FirstLine.ReplaceInline( TEXT(","), TEXT("") ) + 1;
+		if ( NumberOfColumns != FirstLineColumns )
+			UE_LOG( LogHoudiniNiagara, Error,
+			TEXT("Error while parsing the CSV File. Found %d columns but found %d values in the first line! Some values will have an offset!" ),
+			NumberOfColumns, FirstLineColumns );
     }
     
     // Remove the title row now that it's been processed
@@ -242,52 +282,121 @@ bool UHoudiniCSV::UpdateFromStringArray( TArray<FString>& RawStringArray )
     StringCSVData.Empty();
     StringCSVData.SetNumZeroed( NumberOfLines * NumberOfColumns );
 
-     // Extract all the values from the table to the float & string buffers
+	TimeValues.Empty();
+
+    // Extract all the values from the table to the float & string buffers
     TArray<FString> CurrentParsedLine;
     for ( int rowIdx = 0; rowIdx < NumberOfLines; rowIdx++ )
     {
-	FString CurrentLine = RawStringArray[ rowIdx ];
-	if ( HasPackedVectors )
-	{
-	    // Clean up the packing characters: ()" from the line so it can be parsed properly
-	    CurrentLine.ReplaceInline( TEXT("("), TEXT("") );
-	    CurrentLine.ReplaceInline( TEXT(")"), TEXT("") );
-	    CurrentLine.ReplaceInline( TEXT("\""), TEXT("") );
-	}
+		FString CurrentLine = RawStringArray[ rowIdx ];
+		if ( HasPackedVectors )
+		{
+			// Clean up the packing characters: ()" from the line so it can be parsed properly
+			CurrentLine.ReplaceInline( TEXT("("), TEXT("") );
+			CurrentLine.ReplaceInline( TEXT(")"), TEXT("") );
+			CurrentLine.ReplaceInline( TEXT("\""), TEXT("") );
+		}
+
+		// Parse the current line to an array
+		CurrentParsedLine.Empty();
+		CurrentLine.ParseIntoArray( CurrentParsedLine, TEXT(",") );
+
+		// Check the parsed line and number of columns match
+		if ( NumberOfColumns != CurrentParsedLine.Num() )
+			UE_LOG( LogHoudiniNiagara, Warning,
+			TEXT("Error while parsing the CSV File. Line %d has %d values instead of the expected %d!"),
+			rowIdx + 1, CurrentParsedLine.Num(), NumberOfColumns );
+
+		// Store the CSV Data in the buffers
+		// The data is stored transposed in those buffers
+		int32 CurrentID = -1;
+		for ( int colIdx = 0; colIdx < NumberOfColumns; colIdx++ )
+		{
+			// Get the string value for the current column
+			FString CurrentVal = TEXT("0");
+			if ( CurrentParsedLine.IsValidIndex( colIdx ) )
+				CurrentVal = CurrentParsedLine[ colIdx ];
+			else
+			{
+				UE_LOG( LogHoudiniNiagara, Warning,
+				TEXT("Error while parsing the CSV File. Line %d has an invalid value for column %d!"),
+				rowIdx + 1, colIdx + 1 );
+			}
+
+			// Convert the string value to a float in the buffer
+			float FloatValue = FCString::Atof( *CurrentVal );
+			FloatCSVData[ rowIdx + ( colIdx * NumberOfLines ) ] = FloatValue;
+
+			// Keep the original string value in a buffer too
+			StringCSVData[ rowIdx + ( colIdx * NumberOfLines ) ] = CurrentVal;
+						
+			if ( colIdx == IDColumnIndex )
+				CurrentID = (int32)FloatValue;
+
+			if ( colIdx == TimeColumnIndex )
+				TimeValues.AddUnique( FloatValue );
+		}
+
+		// Look at the particle ID, the max ID will be our number of particles
+		if ( NumberOfParticles <= CurrentID )
+			NumberOfParticles = CurrentID + 1;
+    }
 
-	// Parse the current line to an array
-	CurrentParsedLine.Empty();
-	CurrentLine.ParseIntoArray( CurrentParsedLine, TEXT(",") );
+	if ( NumberOfParticles <= 0 )
+		NumberOfParticles = NumberOfLines;
 
-	// Check the parsed line and number of columns match
-	if ( NumberOfColumns != CurrentParsedLine.Num() )
-	    UE_LOG( LogHoudiniNiagara, Warning,
-		TEXT("Error while parsing the CSV File. Line %d has %d values instead of the expected %d!"),
-		rowIdx + 1, CurrentParsedLine.Num(), NumberOfColumns );
+	// Look for particles specific attributes
+	SpawnTimes.Empty();
+	SpawnTimes.Init( -1.0f, NumberOfParticles );
 
-	// Store the CSV Data in the buffer
-	// The data is stored transposed in those buffers
-	for ( int colIdx = 0; colIdx < NumberOfColumns; colIdx++ )
-	{
-	    // Get the string value for the current column
-	    FString CurrentVal = TEXT("0");
-	    if ( CurrentParsedLine.IsValidIndex( colIdx ) )
-		CurrentVal = CurrentParsedLine[ colIdx ];
-	    else
-	    {
-		UE_LOG( LogHoudiniNiagara, Warning,
-		    TEXT("Error while parsing the CSV File. Line %d has an invalid value for column %d!"),
-		    rowIdx + 1, colIdx + 1 );
-	    }
-
-	    // Convert the string value to a float in the buffer
-	    FloatCSVData[ rowIdx + ( colIdx * NumberOfLines ) ] = FCString::Atof( *CurrentVal );
-
-	    // Keep the original string value in a buffer too
-	    StringCSVData[ rowIdx + ( colIdx * NumberOfLines ) ] = CurrentVal;
-	}	
-    }
+	LifeValues.Empty();
+	LifeValues.Init( -1.0f,  NumberOfParticles );
 
+	float MaxTime = -1.0f;
+	for ( int rowIdx = 0; rowIdx < NumberOfLines; rowIdx++ )
+	{
+		// Get the particle ID
+		int32 CurrentID = rowIdx;
+		if ( IDColumnIndex != INDEX_NONE )
+			CurrentID = (int32)FloatCSVData[ rowIdx + ( IDColumnIndex * NumberOfLines ) ];
+
+		// Get the current time
+		float CurrentTime = 0.0f;
+		if ( TimeColumnIndex != INDEX_NONE )
+			CurrentTime = FloatCSVData[ rowIdx + ( TimeColumnIndex * NumberOfLines ) ];
+
+		if ( LifeColumnIndex != INDEX_NONE )
+		{
+			// Set spawn time and life using life values
+			float CurrentLife = FloatCSVData[ rowIdx + ( LifeColumnIndex * NumberOfLines ) ];
+			if ( SpawnTimes[ CurrentID ] < 0.0f )
+			{
+				SpawnTimes[ CurrentID ] = CurrentTime;
+				LifeValues[ CurrentID ] = CurrentLife;
+			}
+		}
+		else if ( AliveColumnIndex != INDEX_NONE )
+		{
+			// Set spawn time and life using the alive bool values
+			bool CurrentAlive = ( FloatCSVData[ rowIdx + ( AliveColumnIndex * NumberOfLines ) ] == 1.0f );
+			if ( ( SpawnTimes[ CurrentID ] < 0.0f ) && CurrentAlive )
+			{
+				// Spawn time is when the particle is first seen alive
+				SpawnTimes[ CurrentID ] = CurrentTime;
+			}
+			else if ( ( SpawnTimes[ CurrentID ] >= 0.0f ) && !CurrentAlive )
+			{
+				// Life is the difference between spawn time and time of death
+				LifeValues[ CurrentID ] = CurrentTime - SpawnTimes[ CurrentID ];
+			}
+		}
+		else
+		{
+			// No life or alive value, spawn time is the first time we see the particle
+			if ( SpawnTimes [ CurrentID ] == INDEX_NONE )
+				SpawnTimes[ CurrentID ] = CurrentTime;
+		}
+	}
     return true;
 }
 
@@ -296,16 +405,16 @@ bool UHoudiniCSV::UpdateFromStringArray( TArray<FString>& RawStringArray )
 bool UHoudiniCSV::GetCSVFloatValue( const int32& lineIndex, const int32& colIndex, float& value )
 {
     if ( lineIndex < 0 || lineIndex >= NumberOfLines )
-	return false;
+		return false;
 
     if ( colIndex < 0 || colIndex >= NumberOfColumns )
-	return false;
+		return false;
 
     int32 Index = lineIndex + ( colIndex * NumberOfLines );
     if ( FloatCSVData.IsValidIndex( Index ) )
     {
-	value = FloatCSVData[ Index ];
-	return true;
+		value = FloatCSVData[ Index ];
+		return true;
     }
 
     return false;
@@ -315,16 +424,16 @@ bool UHoudiniCSV::GetCSVFloatValue( const int32& lineIndex, const int32& colInde
 bool UHoudiniCSV::GetCSVStringValue( const int32& lineIndex, const int32& colIndex, FString& value )
 {
     if ( lineIndex < 0 || lineIndex >= NumberOfLines )
-	return false;
+		return false;
 
     if ( colIndex < 0 || colIndex >= NumberOfColumns )
-	return false;
+		return false;
 
     int32 Index = lineIndex + ( colIndex * NumberOfLines );
     if ( StringCSVData.IsValidIndex( Index ) )
     {
-	value = StringCSVData[ Index ];
-	return true;
+		value = StringCSVData[ Index ];
+		return true;
     }
 
     return false;
@@ -335,23 +444,23 @@ bool UHoudiniCSV::GetCSVVectorValue( const int32& lineIndex, const int32& colInd
 {
     FVector V = FVector::ZeroVector;
     if ( !GetCSVFloatValue( lineIndex, colIndex, V.X ) )
-	return false;
+		return false;
 
     if ( !GetCSVFloatValue( lineIndex, colIndex + 1, V.Y ) )
-	return false;
+		return false;
 
     if ( !GetCSVFloatValue( lineIndex, colIndex + 2, V.Z ) )
-	return false;
+		return false;
 
     if ( DoScale )
-	V *= 100.0f;
+		V *= 100.0f;
 
     value = V;
 
     if ( DoSwap )
     {
-	value.Y = V.Z;
-	value.Z = V.Y;
+		value.Y = V.Z;
+		value.Z = V.Y;
     }
 
     return true;
@@ -362,7 +471,7 @@ bool UHoudiniCSV::GetCSVPositionValue( const int32& lineIndex, FVector& value )
 {
     FVector V = FVector::ZeroVector;
     if ( !GetCSVVectorValue( lineIndex, PositionColumnIndex, V, true, true ) )
-	return false;
+		return false;
 
     value = V;
 
@@ -374,7 +483,7 @@ bool UHoudiniCSV::GetCSVNormalValue( const int32& lineIndex, FVector& value )
 {
     FVector V = FVector::ZeroVector;
     if ( !GetCSVVectorValue(lineIndex, NormalColumnIndex, V, true, false ) )
-	return false;
+		return false;
 
     value = V;
 
@@ -386,7 +495,7 @@ bool UHoudiniCSV::GetCSVTimeValue( const int32& lineIndex, float& value )
 {
     float temp;
     if ( !GetCSVFloatValue( lineIndex, TimeColumnIndex, temp ) )
-	return false;
+		return false;
 
     value = temp;
 
@@ -396,9 +505,18 @@ bool UHoudiniCSV::GetCSVTimeValue( const int32& lineIndex, float& value )
 // Returns the number of points found in the CSV file
 int32 UHoudiniCSV::GetNumberOfPointsInCSV()
 {
+	if ( IDColumnIndex != INDEX_NONE )
+		return NumberOfParticles;
+
     return NumberOfLines;
 }
 
+// Returns the number of points found in the CSV file
+int32 UHoudiniCSV::GetNumberOfLinesInCSV()
+{
+	return NumberOfLines;
+}
+
 // Get the last index of the particles with a time value smaller or equal to desiredTime
 // If the CSV file doesn't have time informations, returns false and set the LastIndex to the last particle
 // If desiredTime is smaller than the first particle, LastIndex will be set to -1
@@ -408,55 +526,234 @@ bool UHoudiniCSV::GetLastParticleIndexAtTime( const float& desiredTime, int32& l
     // If we dont have proper time info, always return the last line
     if ( TimeColumnIndex < 0 || TimeColumnIndex >= NumberOfColumns )
     {
-	lastIndex = NumberOfLines - 1;
-	return false;
+		lastIndex = NumberOfLines - 1;
+		return false;
     }
 
+	/*
     float temp_time = 0.0f;
+	if ( SpawnTimes.IsValidIndex(NumberOfParticles - 1) )
+	{
+		if ( SpawnTimes[ NumberOfParticles - 1 ] < desiredTime )
+		{
+			// We didn't find a suitable index because the desired time is higher than our last time value
+			lastIndex = NumberOfLines - 1;
+			return true;
+		}
+	}
+
+	// Iterates through all the particles
+	lastIndex = INDEX_NONE;
+	for (int32 n = 0; n < NumberOfParticles; n++)
+	{
+		temp_time = SpawnTimes[ NumberOfParticles - 1 ];
+
+
+
+		if (GetCSVTimeValue(n, temp_time))
+		{
+			if (temp_time == desiredTime)
+			{
+				lastIndex = n;
+			}
+			else if (temp_time > desiredTime)
+			{
+				lastIndex = n - 1;
+				return true;
+			}
+		}
+	}
+	*/
+	float temp_time = 0.0f;
     // Check first if we have anything to spawn at the current time by looking at the last value
     if ( GetCSVTimeValue( NumberOfLines - 1, temp_time ) && temp_time < desiredTime )
     {
-	// We didn't find a suitable index because the desired time is higher than our last time value
-	lastIndex = NumberOfLines - 1;
-	return true;
+		// We didn't find a suitable index because the desired time is higher than our last time value
+		lastIndex = NumberOfLines - 1;
+		return true;
     }	
 
     // Iterates through all the particles
     lastIndex = INDEX_NONE;
     for ( int32 n = 0; n < NumberOfLines; n++ )
     {
-	if ( GetCSVTimeValue( n, temp_time ) )
-	{
-	    if ( temp_time == desiredTime )
-	    {
-		lastIndex = n;
-	    }
-	    else if ( temp_time > desiredTime )
-	    {
-		lastIndex = n - 1;
-		return true;
-	    }
-	}
+		if (GetCSVTimeValue(n, temp_time))
+		{
+			if (temp_time == desiredTime)
+			{
+				lastIndex = n;
+			}
+			else if (temp_time > desiredTime)
+			{
+				lastIndex = n - 1;
+				return true;
+			}
+		}
     }
 
     // We didn't find a suitable index because the desired time is higher than our last time value
     if ( lastIndex == INDEX_NONE )
-	lastIndex = NumberOfLines - 1;
+		lastIndex = NumberOfLines - 1;
 
     return true;
 }
 
+bool UHoudiniCSV::GetParticleValueAtTime( const int32& ParticleID, const int32& ColumnIndex, const float& desiredTime, float& Value )
+{
+	int32 PrevIndex = -1;
+	int32 NextIndex = -1;
+	float PrevWeight = 1.0f;
+
+	if ( !GetParticleLineIndexAtTime( ParticleID, desiredTime, PrevIndex, NextIndex, PrevWeight ) )
+		return false;
+
+	float PrevValue, NextValue;
+	if ( !GetCSVFloatValue( PrevIndex, ColumnIndex, PrevValue ) )
+		return false;
+	if ( !GetCSVFloatValue( NextIndex, ColumnIndex, NextValue ) )
+		return false;
+
+	Value = FMath::Lerp( PrevValue, NextValue, PrevWeight );
+
+	return true;
+}
+
+bool UHoudiniCSV::GetParticlePositionAtTime( const int32& ParticleID, const float& desiredTime, FVector& Vector )
+{
+	FVector V = FVector::ZeroVector;
+	if ( !GetParticleVectorValueAtTime(ParticleID, PositionColumnIndex, desiredTime, V, true, true ) )
+		return false;
+
+	Vector = V;
+
+	return true;
+}
+
+bool UHoudiniCSV::GetParticleVectorValueAtTime( const int32& ParticleID, const int32& ColumnIndex, const float& desiredTime, FVector& Vector, const bool& DoSwap, const bool& DoScale )
+{
+	int32 PrevIndex = -1;
+	int32 NextIndex = -1;
+	float PrevWeight = 1.0f;
+
+	if ( !GetParticleLineIndexAtTime( ParticleID, desiredTime, PrevIndex, NextIndex, PrevWeight ) )
+		return false;
+
+	FVector PrevVector, NextVector;
+	if ( !GetCSVVectorValue( PrevIndex, ColumnIndex, PrevVector, DoSwap, DoScale ) )
+		return false;
+	if ( !GetCSVVectorValue( NextIndex, ColumnIndex, NextVector, DoSwap, DoScale ) )
+		return false;
+
+	Vector = FMath::Lerp(PrevVector, NextVector, PrevWeight);
+
+	return true;
+}
+
+bool UHoudiniCSV::GetParticleLineIndexAtTime(const int32& ParticleID, const float& desiredTime, int32& PrevIndex, int32& NextIndex, float& PrevWeight )
+{
+	float PrevTime = -1.0f;
+	float NextTime = -1.0f;
+
+	// Invalid Particle ID
+	if ( ParticleID < 0 || ParticleID >= NumberOfParticles )
+		return false;
+
+	// Get the spawn time
+	float ParticleSpawnTime = 0.0f;
+	if ( SpawnTimes.IsValidIndex( ParticleID ) )
+		ParticleSpawnTime = SpawnTimes[ ParticleID ];
+
+	// If particle hasn't spawn before DesiredTime
+	if ( ParticleSpawnTime > desiredTime )
+		return false;
+
+	// Get the life value
+	float ParticleLife = 0.0f;
+	if ( LifeValues.IsValidIndex( ParticleID ) )
+		ParticleLife = LifeValues[ ParticleID ];
+
+	// If particle is dead before DesiredTime
+	if ( ParticleLife > 0.0f && ( ParticleSpawnTime + ParticleLife < desiredTime ) )	
+		return false;
+
+	for ( int32 n = 0; n < NumberOfLines; n++ )
+	{
+		// Get the ID
+		float CurrentLineID = -1;
+		if ( !GetCSVFloatValue( n, IDColumnIndex, CurrentLineID ) )
+			CurrentLineID = n;
+
+		// Not the particle we're looking for
+		if ( CurrentLineID != ParticleID )
+			continue;
+
+		// Get the time
+		float currentTime = -1.0f;
+		if ( GetCSVTimeValue( n, currentTime) )
+		{
+			if ( currentTime == desiredTime )
+			{
+				PrevIndex = n;
+				NextIndex = n;
+				PrevWeight = 1.0f;
+				return true;
+			}
+			else if ( currentTime < desiredTime )
+			{
+				if ( PrevTime == -1.0f || PrevTime < currentTime )
+				{
+					PrevIndex = n;
+					PrevTime = currentTime;
+				}
+			}
+			else
+			{
+				if ( NextTime == -1.0f || NextTime > currentTime)
+				{
+					NextIndex = n;
+					NextTime = currentTime;
+
+					// TODO: since the csv is sorted by time, we can break now
+					break;
+				}
+			}
+		}
+	}
+
+	if ( PrevIndex < 0 && NextIndex < 0 )
+		return false;
+
+	if ( PrevIndex < 0 )
+	{
+		PrevWeight = 0.0f;
+		return true;
+	}
+
+	if ( NextIndex < 0 )
+	{
+		PrevWeight = 1.0f;
+		return true;
+	}
+
+	// Calculate the weight
+	PrevWeight = ( ( desiredTime - PrevTime) / ( NextTime - PrevTime ) );
+
+	return true;
+}
+
+
+
 // Returns the column index for a given string
 bool UHoudiniCSV::GetColumnIndexFromString( const FString& ColumnTitle, int32& ColumnIndex )
 {
     if ( !TitleRowArray.Find( ColumnTitle, ColumnIndex ) )
-	return true;
+		return true;
 
     // Handle packed positions/normals here
     if ( ColumnTitle.Equals( "P" ) )
-	return TitleRowArray.Find( TEXT( "Px" ), ColumnIndex );
+		return TitleRowArray.Find( TEXT( "Px" ), ColumnIndex );
     else if ( ColumnTitle.Equals( "N" ) )
-	return TitleRowArray.Find( TEXT( "Nx" ), ColumnIndex );
+		return TitleRowArray.Find( TEXT( "Nx" ), ColumnIndex );
 
     return false;
 }
@@ -466,7 +763,7 @@ bool UHoudiniCSV::GetCSVFloatValue( const int32& lineIndex, const FString& Colum
 {
     int32 ColIndex = -1;
     if ( !GetColumnIndexFromString( ColumnTitle, ColIndex ) )
-	return false;
+		return false;
 
     return GetCSVFloatValue( lineIndex, ColIndex, value );
 }
@@ -477,17 +774,19 @@ bool UHoudiniCSV::GetCSVStringValue( const int32& lineIndex, const FString& Colu
 {
     int32 ColIndex = -1;
     if ( !GetColumnIndexFromString( ColumnTitle, ColIndex ) )
-	return false;
+		return false;
 
     return GetCSVStringValue( lineIndex, ColIndex, value );
 }
 
+// 
+
 #if WITH_EDITORONLY_DATA
 void UHoudiniCSV::PostInitProperties()
 {
     if ( !HasAnyFlags( RF_ClassDefaultObject ) )
     {
-	AssetImportData = NewObject<UAssetImportData>(this, TEXT("AssetImportData"));
+		AssetImportData = NewObject<UAssetImportData>(this, TEXT("AssetImportData"));
     }
 
     Super::PostInitProperties();
diff --git a/Engine/Plugins/FX/HoudiniNiagara/Source/HoudiniNiagara/Private/HoudiniNiagara.cpp b/Engine/Plugins/FX/HoudiniNiagara/Source/HoudiniNiagara/Private/HoudiniNiagara.cpp
--- a/Engine/Plugins/FX/HoudiniNiagara/Source/HoudiniNiagara/Private/HoudiniNiagara.cpp
+++ b/Engine/Plugins/FX/HoudiniNiagara/Source/HoudiniNiagara/Private/HoudiniNiagara.cpp
@@ -1,5 +1,5 @@
 /*
-* Copyright (c) <2017> Side Effects Software Inc.
+* Copyright (c) <2018> Side Effects Software Inc.
 *
 * Permission is hereby granted, free of charge, to any person obtaining a copy
 * of this software and associated documentation files (the "Software"), to deal
diff --git a/Engine/Plugins/FX/HoudiniNiagara/Source/HoudiniNiagara/Private/NiagaraDataInterfaceHoudiniCSV.cpp b/Engine/Plugins/FX/HoudiniNiagara/Source/HoudiniNiagara/Private/NiagaraDataInterfaceHoudiniCSV.cpp
--- a/Engine/Plugins/FX/HoudiniNiagara/Source/HoudiniNiagara/Private/NiagaraDataInterfaceHoudiniCSV.cpp
+++ b/Engine/Plugins/FX/HoudiniNiagara/Source/HoudiniNiagara/Private/NiagaraDataInterfaceHoudiniCSV.cpp
@@ -1,5 +1,5 @@
 /*
-* Copyright (c) <2017> Side Effects Software Inc.
+* Copyright (c) <2018> Side Effects Software Inc.
 *
 * Permission is hereby granted, free of charge, to any person obtaining a copy
 * of this software and associated documentation files (the "Software"), to deal
@@ -63,12 +63,12 @@ void UNiagaraDataInterfaceHoudiniCSV::PostEditChangeProperty(struct FPropertyCha
 
     if (PropertyChangedEvent.Property && PropertyChangedEvent.Property->GetFName() == GET_MEMBER_NAME_CHECKED(UNiagaraDataInterfaceHoudiniCSV, CSVFile))
     {
-	Modify();
-	if ( CSVFile )
-	{
-	    GPUBufferDirty = true;
-	    LastSpawnIndex = -1;
-	}
+		Modify();
+		if (CSVFile)
+		{
+			GPUBufferDirty = true;
+			LastSpawnIndex = -1;
+		}
     }
 }
 
@@ -77,13 +77,11 @@ void UNiagaraDataInterfaceHoudiniCSV::PostEditChangeProperty(struct FPropertyCha
 bool UNiagaraDataInterfaceHoudiniCSV::CopyToInternal(UNiagaraDataInterface* Destination) const
 {
     if ( !Super::CopyToInternal( Destination ) )
-    {
-	return false;
-    }
+		return false;
 
     UNiagaraDataInterfaceHoudiniCSV* CastedInterface = CastChecked<UNiagaraDataInterfaceHoudiniCSV>( Destination );
     if ( !CastedInterface )
-	return false;
+		return false;
 
     CastedInterface->CSVFile = CSVFile;
 
@@ -93,16 +91,14 @@ bool UNiagaraDataInterfaceHoudiniCSV::CopyToInternal(UNiagaraDataInterface* Dest
 bool UNiagaraDataInterfaceHoudiniCSV::Equals(const UNiagaraDataInterface* Other) const
 {
     if ( !Super::Equals(Other) )
-    {
-	return false;
-    }
+		return false;
 
     const UNiagaraDataInterfaceHoudiniCSV* OtherHNCSV = CastChecked<UNiagaraDataInterfaceHoudiniCSV>(Other);
 
     if ( OtherHNCSV != nullptr && OtherHNCSV->CSVFile != nullptr && CSVFile )
     {
-	// Just make sure the two interfaces point to the same file
-	return OtherHNCSV->CSVFile->FileName.Equals( CSVFile->FileName );
+		// Just make sure the two interfaces point to the same file
+		return OtherHNCSV->CSVFile->FileName.Equals( CSVFile->FileName );
     }
 
     return false;
@@ -112,126 +108,202 @@ bool UNiagaraDataInterfaceHoudiniCSV::Equals(const UNiagaraDataInterface* Other)
 void UNiagaraDataInterfaceHoudiniCSV::GetFunctions(TArray<FNiagaraFunctionSignature>& OutFunctions)
 {
     {
-	// GetCSVFloatValue
-	FNiagaraFunctionSignature Sig;
-	Sig.Name = TEXT("GetCSVFloatValue");
-	Sig.bMemberFunction = true;
-	Sig.bRequiresContext = false;
-	Sig.Inputs.Add(FNiagaraVariable(FNiagaraTypeDefinition(GetClass()), TEXT("CSV")));		// CSV in
-	Sig.Inputs.Add(FNiagaraVariable(FNiagaraTypeDefinition::GetFloatDef(), TEXT("Row")));		// Row Index In
-	Sig.Inputs.Add(FNiagaraVariable(FNiagaraTypeDefinition::GetFloatDef(), TEXT("Col")));		// Col Index In
-	Sig.Outputs.Add(FNiagaraVariable(FNiagaraTypeDefinition::GetFloatDef(), TEXT("Value")));	// Float Out
+		// GetCSVFloatValue
+		FNiagaraFunctionSignature Sig;
+		Sig.Name = TEXT("GetCSVFloatValue");
+		Sig.bMemberFunction = true;
+		Sig.bRequiresContext = false;
+		Sig.Inputs.Add(FNiagaraVariable(FNiagaraTypeDefinition(GetClass()), TEXT("CSV")));			// CSV in
+		Sig.Inputs.Add(FNiagaraVariable(FNiagaraTypeDefinition::GetFloatDef(), TEXT("Row")));		// Row Index In
+		Sig.Inputs.Add(FNiagaraVariable(FNiagaraTypeDefinition::GetFloatDef(), TEXT("Col")));		// Col Index In
+		Sig.Outputs.Add(FNiagaraVariable(FNiagaraTypeDefinition::GetFloatDef(), TEXT("Value")));	// Float Out
 
-	OutFunctions.Add(Sig);
+		OutFunctions.Add(Sig);
     }
 
+	/*
     {
-	// GetCSVFloatValueByString
-	FNiagaraFunctionSignature Sig;
-	Sig.Name = TEXT("GetCSVFloatValueByString");
-	Sig.bMemberFunction = true;
-	Sig.bRequiresContext = false;
-	Sig.Inputs.Add(FNiagaraVariable(FNiagaraTypeDefinition(GetClass()), TEXT("CSV")));		// CSV in
-	Sig.Inputs.Add(FNiagaraVariable(FNiagaraTypeDefinition::GetFloatDef(), TEXT("Row")));		// Row Index In
-	Sig.Inputs.Add(FNiagaraVariable(FNiagaraTypeDefinition(GetClass()), TEXT("ColTitle")));	// Col Title In
-	Sig.Outputs.Add(FNiagaraVariable(FNiagaraTypeDefinition::GetFloatDef(), TEXT("Value")));	// Float Out
+		// GetCSVFloatValueByString
+		FNiagaraFunctionSignature Sig;
+		Sig.Name = TEXT("GetCSVFloatValueByString");
+		Sig.bMemberFunction = true;
+		Sig.bRequiresContext = false;
+		Sig.Inputs.Add(FNiagaraVariable(FNiagaraTypeDefinition(GetClass()), TEXT("CSV")));			// CSV in
+		Sig.Inputs.Add(FNiagaraVariable(FNiagaraTypeDefinition::GetFloatDef(), TEXT("Row")));		// Row Index In
+		Sig.Inputs.Add(FNiagaraVariable(FNiagaraTypeDefinition::GetStringDef(), TEXT("ColTitle")));	// Col Title In
+		Sig.Outputs.Add(FNiagaraVariable(FNiagaraTypeDefinition::GetFloatDef(), TEXT("Value")));	// Float Out
 
-	OutFunctions.Add(Sig);
+		OutFunctions.Add(Sig);
     }
+	*/
+
+	{
+		// GetCSVVectorValue
+		FNiagaraFunctionSignature Sig;
+		Sig.Name = TEXT("GetCSVVectorValue");
+		Sig.bMemberFunction = true;
+		Sig.bRequiresContext = false;
+		Sig.Inputs.Add(FNiagaraVariable(FNiagaraTypeDefinition(GetClass()), TEXT("CSV")));			// CSV in
+		Sig.Inputs.Add(FNiagaraVariable(FNiagaraTypeDefinition::GetFloatDef(), TEXT("Row")));		// Row Index In
+		Sig.Inputs.Add(FNiagaraVariable(FNiagaraTypeDefinition::GetFloatDef(), TEXT("Col")));		// Col Index In
+		Sig.Outputs.Add(FNiagaraVariable(FNiagaraTypeDefinition::GetVec3Def(), TEXT("Value")));		// Vector3 Out
+
+		OutFunctions.Add(Sig);
+	}
 
     {
-	// GetCSVPosition
-	FNiagaraFunctionSignature Sig;
-	Sig.Name = TEXT("GetCSVPosition");
-	Sig.bMemberFunction = true;
-	Sig.bRequiresContext = false;
-	Sig.Inputs.Add(FNiagaraVariable(FNiagaraTypeDefinition(GetClass()), TEXT("CSV")));		// CSV in
-	Sig.Inputs.Add(FNiagaraVariable(FNiagaraTypeDefinition::GetFloatDef(), TEXT("N")));		// Point Number In
-	Sig.Outputs.Add(FNiagaraVariable(FNiagaraTypeDefinition::GetVec3Def(), TEXT("Value")));		// Vector3 Out
+		// GetCSVPosition
+		FNiagaraFunctionSignature Sig;
+		Sig.Name = TEXT("GetCSVPosition");
+		Sig.bMemberFunction = true;
+		Sig.bRequiresContext = false;
+		Sig.Inputs.Add(FNiagaraVariable(FNiagaraTypeDefinition(GetClass()), TEXT("CSV")));			// CSV in
+		Sig.Inputs.Add(FNiagaraVariable(FNiagaraTypeDefinition::GetFloatDef(), TEXT("N")));			// Point Number In
+		Sig.Outputs.Add(FNiagaraVariable(FNiagaraTypeDefinition::GetVec3Def(), TEXT("Value")));		// Vector3 Out
 
-	OutFunctions.Add(Sig);
+		OutFunctions.Add(Sig);
     }
 
     {
-	// GetCSVPositionAndTime
-	FNiagaraFunctionSignature Sig;
-	Sig.Name = TEXT("GetCSVPositionAndTime");
-	Sig.bMemberFunction = true;
-	Sig.bRequiresContext = false;
-	Sig.Inputs.Add(FNiagaraVariable(FNiagaraTypeDefinition(GetClass()), TEXT("CSV")));		// CSV in
-	Sig.Inputs.Add(FNiagaraVariable(FNiagaraTypeDefinition::GetFloatDef(), TEXT("N")));		// Point Number In
-	Sig.Outputs.Add(FNiagaraVariable(FNiagaraTypeDefinition::GetVec3Def(), TEXT("Position")));	// Vector3 Out
-	Sig.Outputs.Add(FNiagaraVariable(FNiagaraTypeDefinition::GetFloatDef(), TEXT("Time")));		// float Out
+		// GetCSVPositionAndTime
+		FNiagaraFunctionSignature Sig;
+		Sig.Name = TEXT("GetCSVPositionAndTime");
+		Sig.bMemberFunction = true;
+		Sig.bRequiresContext = false;
+		Sig.Inputs.Add(FNiagaraVariable(FNiagaraTypeDefinition(GetClass()), TEXT("CSV")));			// CSV in
+		Sig.Inputs.Add(FNiagaraVariable(FNiagaraTypeDefinition::GetFloatDef(), TEXT("N")));			// Point Number In
+		Sig.Outputs.Add(FNiagaraVariable(FNiagaraTypeDefinition::GetVec3Def(), TEXT("Position")));	// Vector3 Out
+		Sig.Outputs.Add(FNiagaraVariable(FNiagaraTypeDefinition::GetFloatDef(), TEXT("Time")));		// float Out
 
-	OutFunctions.Add(Sig);
+		OutFunctions.Add(Sig);
     }
 
     {
-	// GetCSVNormal
-	FNiagaraFunctionSignature Sig;
-	Sig.Name = TEXT("GetCSVNormal");
-	Sig.bMemberFunction = true;
-	Sig.bRequiresContext = false;
-	Sig.Inputs.Add(FNiagaraVariable(FNiagaraTypeDefinition(GetClass()), TEXT("CSV")));		// CSV in
-	Sig.Inputs.Add(FNiagaraVariable(FNiagaraTypeDefinition::GetFloatDef(), TEXT("N")));		// Point Number In
-	Sig.Outputs.Add(FNiagaraVariable(FNiagaraTypeDefinition::GetVec3Def(), TEXT("Value")));		// Vector3 Out
+		// GetCSVNormal
+		FNiagaraFunctionSignature Sig;
+		Sig.Name = TEXT("GetCSVNormal");
+		Sig.bMemberFunction = true;
+		Sig.bRequiresContext = false;
+		Sig.Inputs.Add(FNiagaraVariable(FNiagaraTypeDefinition(GetClass()), TEXT("CSV")));			// CSV in
+		Sig.Inputs.Add(FNiagaraVariable(FNiagaraTypeDefinition::GetFloatDef(), TEXT("N")));			// Point Number In
+		Sig.Outputs.Add(FNiagaraVariable(FNiagaraTypeDefinition::GetVec3Def(), TEXT("Value")));		// Vector3 Out
 
-	OutFunctions.Add(Sig);
+		OutFunctions.Add(Sig);
     }
 
     {
-	// GetCSVTime
-	FNiagaraFunctionSignature Sig;
-	Sig.Name = TEXT("GetCSVTime");
-	Sig.bMemberFunction = true;
-	Sig.bRequiresContext = false;
-	Sig.Inputs.Add(FNiagaraVariable(FNiagaraTypeDefinition(GetClass()), TEXT("CSV")));		// CSV in
-	Sig.Inputs.Add(FNiagaraVariable(FNiagaraTypeDefinition::GetFloatDef(), TEXT("N")));		// Point Number In
-	Sig.Outputs.Add(FNiagaraVariable(FNiagaraTypeDefinition::GetFloatDef(), TEXT("Value")));	// Float Out
+		// GetCSVTime
+		FNiagaraFunctionSignature Sig;
+		Sig.Name = TEXT("GetCSVTime");
+		Sig.bMemberFunction = true;
+		Sig.bRequiresContext = false;
+		Sig.Inputs.Add(FNiagaraVariable(FNiagaraTypeDefinition(GetClass()), TEXT("CSV")));			// CSV in
+		Sig.Inputs.Add(FNiagaraVariable(FNiagaraTypeDefinition::GetFloatDef(), TEXT("N")));			// Point Number In
+		Sig.Outputs.Add(FNiagaraVariable(FNiagaraTypeDefinition::GetFloatDef(), TEXT("Value")));	// Float Out
 
-	OutFunctions.Add(Sig);
+		OutFunctions.Add(Sig);
     }
 
     {
-	// GetNumberOfPointsInCSV
-	FNiagaraFunctionSignature Sig;
-	Sig.Name = TEXT("GetNumberOfPointsInCSV");
-	Sig.bMemberFunction = true;
-	Sig.bRequiresContext = false;
-	Sig.Inputs.Add(FNiagaraVariable(FNiagaraTypeDefinition(GetClass()), TEXT("CSV")));		    // CSV in
-	Sig.Outputs.Add(FNiagaraVariable(FNiagaraTypeDefinition::GetIntDef(), TEXT("NumberOfPoints")));     // Int Out
+		// GetNumberOfPointsInCSV
+		FNiagaraFunctionSignature Sig;
+		Sig.Name = TEXT("GetNumberOfPointsInCSV");
+		Sig.bMemberFunction = true;
+		Sig.bRequiresContext = false;
+		Sig.Inputs.Add(FNiagaraVariable(FNiagaraTypeDefinition(GetClass()), TEXT("CSV")));					// CSV in
+		Sig.Outputs.Add(FNiagaraVariable(FNiagaraTypeDefinition::GetIntDef(), TEXT("NumberOfPoints")));     // Int Out
 
-	OutFunctions.Add(Sig);
+		OutFunctions.Add(Sig);
     }
 
     {
-	// GetLastParticleIndexAtTime
-	FNiagaraFunctionSignature Sig;
-	Sig.Name = TEXT("GetLastParticleIndexAtTime");
-	Sig.bMemberFunction = true;
-	Sig.bRequiresContext = false;
-	Sig.Inputs.Add(FNiagaraVariable(FNiagaraTypeDefinition(GetClass()), TEXT("CSV")));		    // CSV in
-	Sig.Inputs.Add(FNiagaraVariable(FNiagaraTypeDefinition::GetFloatDef(), TEXT("Time")));		    // Time in
-	Sig.Outputs.Add(FNiagaraVariable(FNiagaraTypeDefinition::GetIntDef(), TEXT("LastIndex")));	    // Int Out
+		// GetLastParticleIndexAtTime
+		FNiagaraFunctionSignature Sig;
+		Sig.Name = TEXT("GetLastParticleIndexAtTime");
+		Sig.bMemberFunction = true;
+		Sig.bRequiresContext = false;
+		Sig.Inputs.Add(FNiagaraVariable(FNiagaraTypeDefinition(GetClass()), TEXT("CSV")));				 // CSV in
+		Sig.Inputs.Add(FNiagaraVariable(FNiagaraTypeDefinition::GetFloatDef(), TEXT("Time")));		    // Time in
+		Sig.Outputs.Add(FNiagaraVariable(FNiagaraTypeDefinition::GetIntDef(), TEXT("LastIndex")));	    // Int Out
 
-	OutFunctions.Add(Sig);
+		OutFunctions.Add(Sig);
     }
 
     {
-	// GetParticleIndexesToSpawnAtTime
-	FNiagaraFunctionSignature Sig;
-	Sig.Name = TEXT("GetParticleIndexesToSpawnAtTime");
-	Sig.bMemberFunction = true;
-	Sig.bRequiresContext = false;
-	Sig.Inputs.Add(FNiagaraVariable(FNiagaraTypeDefinition(GetClass()), TEXT("CSV")));		    // CSV in
-	Sig.Inputs.Add(FNiagaraVariable(FNiagaraTypeDefinition::GetFloatDef(), TEXT("Time")));		    // Time in
-	Sig.Outputs.Add(FNiagaraVariable(FNiagaraTypeDefinition::GetIntDef(), TEXT("MinIndex")));	    // Int Out
-	Sig.Outputs.Add(FNiagaraVariable(FNiagaraTypeDefinition::GetIntDef(), TEXT("MaxIndex")));	    // Int Out
-	Sig.Outputs.Add(FNiagaraVariable(FNiagaraTypeDefinition::GetIntDef(), TEXT("Count")));		    // Int Out
+		// GetParticleIndexesToSpawnAtTime
+		FNiagaraFunctionSignature Sig;
+		Sig.Name = TEXT("GetParticleIndexesToSpawnAtTime");
+		Sig.bMemberFunction = true;
+		Sig.bRequiresContext = false;
+		Sig.Inputs.Add(FNiagaraVariable(FNiagaraTypeDefinition(GetClass()), TEXT("CSV")));				// CSV in
+		Sig.Inputs.Add(FNiagaraVariable(FNiagaraTypeDefinition::GetFloatDef(), TEXT("Time")));		    // Time in
+		Sig.Outputs.Add(FNiagaraVariable(FNiagaraTypeDefinition::GetIntDef(), TEXT("MinIndex")));	    // Int Out
+		Sig.Outputs.Add(FNiagaraVariable(FNiagaraTypeDefinition::GetIntDef(), TEXT("MaxIndex")));	    // Int Out
+		Sig.Outputs.Add(FNiagaraVariable(FNiagaraTypeDefinition::GetIntDef(), TEXT("Count")));		    // Int Out
 
-	OutFunctions.Add(Sig);
+		OutFunctions.Add(Sig);
     }
 
+	{
+		// GetParticleIndexesAtTime
+		FNiagaraFunctionSignature Sig;
+		Sig.Name = TEXT("GetParticleIndexesAtTime");
+		Sig.bMemberFunction = true;
+		Sig.bRequiresContext = false;
+		Sig.Inputs.Add(FNiagaraVariable(FNiagaraTypeDefinition(GetClass()), TEXT("CSV")));				// CSV in
+		Sig.Inputs.Add(FNiagaraVariable(FNiagaraTypeDefinition::GetFloatDef(), TEXT("N")));				// Point Number In
+		Sig.Inputs.Add(FNiagaraVariable(FNiagaraTypeDefinition::GetFloatDef(), TEXT("Time")));		    // Time in
+		Sig.Outputs.Add(FNiagaraVariable(FNiagaraTypeDefinition::GetIntDef(), TEXT("PreviousLineIndex"))); // Int Out
+		Sig.Outputs.Add(FNiagaraVariable(FNiagaraTypeDefinition::GetIntDef(), TEXT("NextLineIndex")));	// Int Out
+		Sig.Outputs.Add(FNiagaraVariable(FNiagaraTypeDefinition::GetFloatDef(), TEXT("PrevWeight")));	// Float Out
+
+		OutFunctions.Add(Sig);
+	}
+
+	{
+		// GetParticlePositionAtTime
+		FNiagaraFunctionSignature Sig;
+		Sig.Name = TEXT("GetParticlePositionAtTime");
+		Sig.bMemberFunction = true;
+		Sig.bRequiresContext = false;
+		Sig.Inputs.Add(FNiagaraVariable(FNiagaraTypeDefinition(GetClass()), TEXT("CSV")));				// CSV in
+		Sig.Inputs.Add(FNiagaraVariable(FNiagaraTypeDefinition::GetFloatDef(), TEXT("N")));				// Point Number In
+		Sig.Inputs.Add(FNiagaraVariable(FNiagaraTypeDefinition::GetFloatDef(), TEXT("Time")));		    // Time in
+		Sig.Outputs.Add(FNiagaraVariable(FNiagaraTypeDefinition::GetVec3Def(), TEXT("Position")));		// Vector3 Out
+
+		OutFunctions.Add(Sig);
+	}
+
+	{
+		// GetParticleValueAtTime
+		FNiagaraFunctionSignature Sig;
+		Sig.Name = TEXT("GetParticleValueAtTime");
+		Sig.bMemberFunction = true;
+		Sig.bRequiresContext = false;
+		Sig.Inputs.Add(FNiagaraVariable(FNiagaraTypeDefinition(GetClass()), TEXT("CSV")));				// CSV in
+		Sig.Inputs.Add(FNiagaraVariable(FNiagaraTypeDefinition::GetFloatDef(), TEXT("N")));				// Point Number In		
+		Sig.Inputs.Add(FNiagaraVariable(FNiagaraTypeDefinition::GetFloatDef(), TEXT("Time")));		    // Time in
+		Sig.Inputs.Add(FNiagaraVariable(FNiagaraTypeDefinition::GetFloatDef(), TEXT("Col")));			// Col Index In
+		Sig.Outputs.Add(FNiagaraVariable(FNiagaraTypeDefinition::GetFloatDef(), TEXT("Value")));		// Float Out
+
+		OutFunctions.Add(Sig);
+	}
+
+	{
+		// GetParticleVectorValueAtTime
+		FNiagaraFunctionSignature Sig;
+		Sig.Name = TEXT("GetParticleVectorValueAtTime");
+		Sig.bMemberFunction = true;
+		Sig.bRequiresContext = false;
+		Sig.Inputs.Add(FNiagaraVariable(FNiagaraTypeDefinition(GetClass()), TEXT("CSV")));				// CSV in
+		Sig.Inputs.Add(FNiagaraVariable(FNiagaraTypeDefinition::GetFloatDef(), TEXT("N")));				// Point Number In
+		Sig.Inputs.Add(FNiagaraVariable(FNiagaraTypeDefinition::GetFloatDef(), TEXT("Time")));		    // Time in
+		Sig.Inputs.Add(FNiagaraVariable(FNiagaraTypeDefinition::GetFloatDef(), TEXT("Col")));			// Col Index In
+		Sig.Outputs.Add(FNiagaraVariable(FNiagaraTypeDefinition::GetVec3Def(), TEXT("Value")));			// Vector3 Out
+
+		OutFunctions.Add(Sig);
+	}
+
 }
 
 // build the shader function HLSL; function name is passed in, as it's defined per-DI; that way, configuration could change
@@ -318,7 +390,7 @@ void UNiagaraDataInterfaceHoudiniCSV::GetParameterDefinitionHLSL(FNiagaraDataInt
 }
 
 DEFINE_NDI_RAW_FUNC_BINDER(UNiagaraDataInterfaceHoudiniCSV, GetCSVFloatValue);
-DEFINE_NDI_RAW_FUNC_BINDER(UNiagaraDataInterfaceHoudiniCSV, GetCSVFloatValueByString);
+//DEFINE_NDI_RAW_FUNC_BINDER(UNiagaraDataInterfaceHoudiniCSV, GetCSVFloatValueByString);
 DEFINE_NDI_RAW_FUNC_BINDER(UNiagaraDataInterfaceHoudiniCSV, GetCSVVectorValue);
 DEFINE_NDI_RAW_FUNC_BINDER(UNiagaraDataInterfaceHoudiniCSV, GetCSVPosition);
 DEFINE_NDI_RAW_FUNC_BINDER(UNiagaraDataInterfaceHoudiniCSV, GetCSVNormal);
@@ -326,54 +398,76 @@ DEFINE_NDI_RAW_FUNC_BINDER(UNiagaraDataInterfaceHoudiniCSV, GetCSVTime);
 DEFINE_NDI_RAW_FUNC_BINDER(UNiagaraDataInterfaceHoudiniCSV, GetCSVPositionAndTime);
 DEFINE_NDI_RAW_FUNC_BINDER(UNiagaraDataInterfaceHoudiniCSV, GetLastParticleIndexAtTime);
 DEFINE_NDI_RAW_FUNC_BINDER(UNiagaraDataInterfaceHoudiniCSV, GetParticleIndexesToSpawnAtTime);
+DEFINE_NDI_RAW_FUNC_BINDER(UNiagaraDataInterfaceHoudiniCSV, GetParticleIndexesAtTime);
+DEFINE_NDI_RAW_FUNC_BINDER(UNiagaraDataInterfaceHoudiniCSV, GetParticlePositionAtTime);
+DEFINE_NDI_RAW_FUNC_BINDER(UNiagaraDataInterfaceHoudiniCSV, GetParticleValueAtTime);
+DEFINE_NDI_RAW_FUNC_BINDER(UNiagaraDataInterfaceHoudiniCSV, GetParticleVectorValueAtTime);
 void UNiagaraDataInterfaceHoudiniCSV::GetVMExternalFunction(const FVMExternalFunctionBindingInfo& BindingInfo, void* InstanceData, FVMExternalFunction &OutFunc)
 {
     if (BindingInfo.Name == TEXT("GetCSVFloatValue") && BindingInfo.GetNumInputs() == 2 && BindingInfo.GetNumOutputs() == 1)
     {
-	TNDIParamBinder<0, float, TNDIParamBinder<1, float, NDI_RAW_FUNC_BINDER(UNiagaraDataInterfaceHoudiniCSV, GetCSVFloatValue)>>::Bind(this, BindingInfo, InstanceData, OutFunc);
+		TNDIParamBinder<0, float, TNDIParamBinder<1, float, NDI_RAW_FUNC_BINDER(UNiagaraDataInterfaceHoudiniCSV, GetCSVFloatValue)>>::Bind(this, BindingInfo, InstanceData, OutFunc);
     }
-    else if (BindingInfo.Name == TEXT("GetCSVFloatValueByString") && BindingInfo.GetNumInputs() == 2 && BindingInfo.GetNumOutputs() == 1)
+    /*else if (BindingInfo.Name == TEXT("GetCSVFloatValueByString") && BindingInfo.GetNumInputs() == 2 && BindingInfo.GetNumOutputs() == 1)
     {
-	TNDIParamBinder<0, float, TNDIParamBinder<1, FString, NDI_RAW_FUNC_BINDER(UNiagaraDataInterfaceHoudiniCSV, GetCSVFloatValueByString)>>::Bind(this, BindingInfo, InstanceData, OutFunc);
-    }
-    else if (BindingInfo.Name == TEXT("GetCSVVectorValue") && BindingInfo.GetNumInputs() == 2 && BindingInfo.GetNumOutputs() == 1)
+		TNDIParamBinder<0, float, TNDIParamBinder<1, FString, NDI_RAW_FUNC_BINDER(UNiagaraDataInterfaceHoudiniCSV, GetCSVFloatValueByString)>>::Bind(this, BindingInfo, InstanceData, OutFunc);
+    }*/
+    else if (BindingInfo.Name == TEXT("GetCSVVectorValue") && BindingInfo.GetNumInputs() == 2 && BindingInfo.GetNumOutputs() == 3)
     {
-	TNDIParamBinder<0, float, TNDIParamBinder<1, float, NDI_RAW_FUNC_BINDER(UNiagaraDataInterfaceHoudiniCSV, GetCSVVectorValue)>>::Bind(this, BindingInfo, InstanceData, OutFunc);
+		TNDIParamBinder<0, float, TNDIParamBinder<1, float, NDI_RAW_FUNC_BINDER(UNiagaraDataInterfaceHoudiniCSV, GetCSVVectorValue)>>::Bind(this, BindingInfo, InstanceData, OutFunc);
     }    
     else if (BindingInfo.Name == TEXT("GetCSVPosition") && BindingInfo.GetNumInputs() == 1 && BindingInfo.GetNumOutputs() == 3)
     {
-	TNDIParamBinder<0, float, NDI_RAW_FUNC_BINDER(UNiagaraDataInterfaceHoudiniCSV, GetCSVPosition)>::Bind(this, BindingInfo, InstanceData, OutFunc);
+		TNDIParamBinder<0, float, NDI_RAW_FUNC_BINDER(UNiagaraDataInterfaceHoudiniCSV, GetCSVPosition)>::Bind(this, BindingInfo, InstanceData, OutFunc);
     }
     else if (BindingInfo.Name == TEXT("GetCSVNormal") && BindingInfo.GetNumInputs() == 1 && BindingInfo.GetNumOutputs() == 3)
     {
-	TNDIParamBinder<0, float, NDI_RAW_FUNC_BINDER(UNiagaraDataInterfaceHoudiniCSV, GetCSVNormal)>::Bind(this, BindingInfo, InstanceData, OutFunc);
+		TNDIParamBinder<0, float, NDI_RAW_FUNC_BINDER(UNiagaraDataInterfaceHoudiniCSV, GetCSVNormal)>::Bind(this, BindingInfo, InstanceData, OutFunc);
     }
     else if (BindingInfo.Name == TEXT("GetCSVTime") && BindingInfo.GetNumInputs() == 1 && BindingInfo.GetNumOutputs() == 1)
     {
-	TNDIParamBinder<0, float, NDI_RAW_FUNC_BINDER(UNiagaraDataInterfaceHoudiniCSV, GetCSVTime)>::Bind(this, BindingInfo, InstanceData, OutFunc);
+		TNDIParamBinder<0, float, NDI_RAW_FUNC_BINDER(UNiagaraDataInterfaceHoudiniCSV, GetCSVTime)>::Bind(this, BindingInfo, InstanceData, OutFunc);
     }
     else if (BindingInfo.Name == TEXT("GetCSVPositionAndTime") && BindingInfo.GetNumInputs() == 1 && BindingInfo.GetNumOutputs() == 4)
     {
-	TNDIParamBinder<0, float, NDI_RAW_FUNC_BINDER(UNiagaraDataInterfaceHoudiniCSV, GetCSVPositionAndTime)>::Bind(this, BindingInfo, InstanceData, OutFunc);
+		TNDIParamBinder<0, float, NDI_RAW_FUNC_BINDER(UNiagaraDataInterfaceHoudiniCSV, GetCSVPositionAndTime)>::Bind(this, BindingInfo, InstanceData, OutFunc);
     }
     else if ( BindingInfo.Name == TEXT("GetNumberOfPointsInCSV") && BindingInfo.GetNumInputs() == 0 && BindingInfo.GetNumOutputs() == 1 )
     {
-	OutFunc = FVMExternalFunction::CreateUObject(this, &UNiagaraDataInterfaceHoudiniCSV::GetNumberOfPointsInCSV);
+		OutFunc = FVMExternalFunction::CreateUObject(this, &UNiagaraDataInterfaceHoudiniCSV::GetNumberOfPointsInCSV);
     }
     else if (BindingInfo.Name == TEXT("GetLastParticleIndexAtTime") && BindingInfo.GetNumInputs() == 1 && BindingInfo.GetNumOutputs() == 1)
     {
-	TNDIParamBinder<0, float, NDI_RAW_FUNC_BINDER(UNiagaraDataInterfaceHoudiniCSV, GetLastParticleIndexAtTime)>::Bind(this, BindingInfo, InstanceData, OutFunc);
+		TNDIParamBinder<0, float, NDI_RAW_FUNC_BINDER(UNiagaraDataInterfaceHoudiniCSV, GetLastParticleIndexAtTime)>::Bind(this, BindingInfo, InstanceData, OutFunc);
     }
     else if (BindingInfo.Name == TEXT("GetParticleIndexesToSpawnAtTime") && BindingInfo.GetNumInputs() == 1 && BindingInfo.GetNumOutputs() == 3)
     {
-	TNDIParamBinder<0, float, NDI_RAW_FUNC_BINDER(UNiagaraDataInterfaceHoudiniCSV, GetParticleIndexesToSpawnAtTime)>::Bind(this, BindingInfo, InstanceData, OutFunc);
+		TNDIParamBinder<0, float, NDI_RAW_FUNC_BINDER(UNiagaraDataInterfaceHoudiniCSV, GetParticleIndexesToSpawnAtTime)>::Bind(this, BindingInfo, InstanceData, OutFunc);
     }
+	else if (BindingInfo.Name == TEXT("GetParticleIndexesAtTime") && BindingInfo.GetNumInputs() == 2 && BindingInfo.GetNumOutputs() == 3)
+	{
+		TNDIParamBinder<0, float, TNDIParamBinder<1, float, NDI_RAW_FUNC_BINDER(UNiagaraDataInterfaceHoudiniCSV, GetParticleIndexesAtTime)>>::Bind(this, BindingInfo, InstanceData, OutFunc);
+	}
+	else if (BindingInfo.Name == TEXT("GetParticlePositionAtTime") && BindingInfo.GetNumInputs() == 2 && BindingInfo.GetNumOutputs() == 3)
+	{
+		TNDIParamBinder<0, float, TNDIParamBinder<1, float, NDI_RAW_FUNC_BINDER(UNiagaraDataInterfaceHoudiniCSV, GetParticlePositionAtTime)>>::Bind(this, BindingInfo, InstanceData, OutFunc);
+	}
+
+	else if (BindingInfo.Name == TEXT("GetParticleValueAtTime") && BindingInfo.GetNumInputs() == 3 && BindingInfo.GetNumOutputs() == 1)
+	{
+		TNDIParamBinder<0, float, TNDIParamBinder<1, float, TNDIParamBinder<2, float, NDI_RAW_FUNC_BINDER(UNiagaraDataInterfaceHoudiniCSV, GetParticleValueAtTime)>>>::Bind(this, BindingInfo, InstanceData, OutFunc);
+	}
+
+	else if (BindingInfo.Name == TEXT("GetParticleVectorValueAtTime") && BindingInfo.GetNumInputs() == 2 && BindingInfo.GetNumOutputs() == 3)
+	{
+		TNDIParamBinder<0, float, TNDIParamBinder<1, float, TNDIParamBinder<2, float, NDI_RAW_FUNC_BINDER(UNiagaraDataInterfaceHoudiniCSV, GetParticleVectorValueAtTime)>>>::Bind(this, BindingInfo, InstanceData, OutFunc);
+	}
     else
     {
-	UE_LOG( LogHoudiniNiagara, Error, 
+		UE_LOG( LogHoudiniNiagara, Error, 
 	    TEXT( "Could not find data interface function:\n\tName: %s\n\tInputs: %i\n\tOutputs: %i" ),
 	    *BindingInfo.Name.ToString(), BindingInfo.GetNumInputs(), BindingInfo.GetNumOutputs() );
-	OutFunc = FVMExternalFunction();
+		OutFunc = FVMExternalFunction();
     }
 }
 
@@ -387,17 +481,17 @@ void UNiagaraDataInterfaceHoudiniCSV::GetCSVFloatValue(FVectorVMContext& Context
 
     for ( int32 i = 0; i < Context.NumInstances; ++i )
     {
-	float row = RowParam.Get();
-	float col = ColParam.Get();
+		float row = RowParam.Get();
+		float col = ColParam.Get();
 	
-	float value = 0.0f;
-	if ( CSVFile )
-	    CSVFile->GetCSVFloatValue( row, col, value );
+		float value = 0.0f;
+		if ( CSVFile )
+			CSVFile->GetCSVFloatValue( row, col, value );
 
-	*OutValue.GetDest() = value;
-	RowParam.Advance();
-	ColParam.Advance();
-	OutValue.Advance();
+		*OutValue.GetDest() = value;
+		RowParam.Advance();
+		ColParam.Advance();
+		OutValue.Advance();
     }
 }
 
@@ -407,24 +501,34 @@ void UNiagaraDataInterfaceHoudiniCSV::GetCSVVectorValue( FVectorVMContext& Conte
     RowParamType RowParam(Context);
     ColParamType ColParam(Context);
 
-    FRegisterHandler<FVector> OutValue(Context);
+    //FRegisterHandler<FVector> OutValue(Context);
+	FRegisterHandler<float> OutVectorX(Context);
+	FRegisterHandler<float> OutVectorY(Context);
+	FRegisterHandler<float> OutVectorZ(Context);
 
     for (int32 i = 0; i < Context.NumInstances; ++i)
     {
-	float row = RowParam.Get();
-	float col = ColParam.Get();
+		float row = RowParam.Get();
+		float col = ColParam.Get();
 
-	FVector value = FVector::ZeroVector;
-	if ( CSVFile )
-	    CSVFile->GetCSVVectorValue(row, col, value);
+		FVector V = FVector::ZeroVector;
+		if ( CSVFile )
+			CSVFile->GetCSVVectorValue(row, col, V);
 
-	*OutValue.GetDest() = value;
-	RowParam.Advance();
-	ColParam.Advance();
-	OutValue.Advance();
+		//*OutValue.GetDest() = value;
+		*OutVectorX.GetDest() = V.X;
+		*OutVectorY.GetDest() = V.Y;
+		*OutVectorZ.GetDest() = V.Z;
+
+		RowParam.Advance();
+		ColParam.Advance();
+		//OutValue.Advance();
+		OutVectorX.Advance();
+		OutVectorY.Advance();
+		OutVectorZ.Advance();
     }
 }
-
+/*
 template<typename RowParamType, typename ColTitleParamType>
 void UNiagaraDataInterfaceHoudiniCSV::GetCSVFloatValueByString(FVectorVMContext& Context)
 {
@@ -435,19 +539,20 @@ void UNiagaraDataInterfaceHoudiniCSV::GetCSVFloatValueByString(FVectorVMContext&
 
     for ( int32 i = 0; i < Context.NumInstances; ++i )
     {
-	float row = RowParam.Get();
-	FString colTitle = ColTitleParam.Get();
+		float row = RowParam.Get();
+		FString colTitle = ColTitleParam.Get();
 	
-	float value = 0.0f;
-	if ( CSVFile )
-	    CSVFile->GetCSVFloatValue( row, colTitle, value );
+		float value = 0.0f;
+		if ( CSVFile )
+			CSVFile->GetCSVFloatValue( row, colTitle, value );
 
-	*OutValue.GetDest() = value;
-	RowParam.Advance();
-	ColTitleParam.Advance();
-	OutValue.Advance();
+		*OutValue.GetDest() = value;
+		RowParam.Advance();
+		ColTitleParam.Advance();
+		OutValue.Advance();
     }
 }
+*/
 
 template<typename NParamType>
 void UNiagaraDataInterfaceHoudiniCSV::GetCSVPosition(FVectorVMContext& Context)
@@ -459,19 +564,19 @@ void UNiagaraDataInterfaceHoudiniCSV::GetCSVPosition(FVectorVMContext& Context)
 
     for (int32 i = 0; i < Context.NumInstances; ++i)
     {
-	float N = NParam.Get();
+		float N = NParam.Get();
 
-	FVector V = FVector::ZeroVector;
-	if ( CSVFile )
-	    CSVFile->GetCSVPositionValue( N, V );
+		FVector V = FVector::ZeroVector;
+		if (CSVFile)
+			CSVFile->GetCSVPositionValue(N, V);
 
-	*OutSampleX.GetDest() = V.X;
-	*OutSampleY.GetDest() = V.Y;
-	*OutSampleZ.GetDest() = V.Z;
-	NParam.Advance();
-	OutSampleX.Advance();
-	OutSampleY.Advance();
-	OutSampleZ.Advance();
+		*OutSampleX.GetDest() = V.X;
+		*OutSampleY.GetDest() = V.Y;
+		*OutSampleZ.GetDest() = V.Z;
+		NParam.Advance();
+		OutSampleX.Advance();
+		OutSampleY.Advance();
+		OutSampleZ.Advance();
     }
 }
 
@@ -485,19 +590,19 @@ void UNiagaraDataInterfaceHoudiniCSV::GetCSVNormal(FVectorVMContext& Context)
 
     for (int32 i = 0; i < Context.NumInstances; ++i)
     {
-	float N = NParam.Get();
+		float N = NParam.Get();
 
-	FVector V = FVector::ZeroVector;
-	if ( CSVFile )
-	    CSVFile->GetCSVNormalValue( N, V );
+		FVector V = FVector::ZeroVector;
+		if (CSVFile)
+			CSVFile->GetCSVNormalValue(N, V);
 
-	*OutSampleX.GetDest() = V.X;
-	*OutSampleY.GetDest() = V.Y;
-	*OutSampleZ.GetDest() = V.Z;
-	NParam.Advance();
-	OutSampleX.Advance();
-	OutSampleY.Advance();
-	OutSampleZ.Advance();
+		*OutSampleX.GetDest() = V.X;
+		*OutSampleY.GetDest() = V.Y;
+		*OutSampleZ.GetDest() = V.Z;
+		NParam.Advance();
+		OutSampleX.Advance();
+		OutSampleY.Advance();
+		OutSampleZ.Advance();
     }
 }
 
@@ -509,15 +614,15 @@ void UNiagaraDataInterfaceHoudiniCSV::GetCSVTime(FVectorVMContext& Context)
 
     for (int32 i = 0; i < Context.NumInstances; ++i)
     {
-	float N = NParam.Get();
+		float N = NParam.Get();
 
-	float value = 0.0f;
-	if ( CSVFile )
-	    CSVFile->GetCSVTimeValue( N, value );
+		float value = 0.0f;
+		if (CSVFile)
+			CSVFile->GetCSVTimeValue(N, value);
 
-	*OutValue.GetDest() = value;
-	NParam.Advance();
-	OutValue.Advance();
+		*OutValue.GetDest() = value;
+		NParam.Advance();
+		OutValue.Advance();
     }
 }
 
@@ -530,15 +635,15 @@ void UNiagaraDataInterfaceHoudiniCSV::GetLastParticleIndexAtTime(FVectorVMContex
 
     for (int32 i = 0; i < Context.NumInstances; ++i)
     {
-	float t = TimeParam.Get();
+		float t = TimeParam.Get();
 
-	int32 value = 0;
-	if ( CSVFile )
-	    CSVFile->GetLastParticleIndexAtTime( t, value );
+		int32 value = 0;
+		if ( CSVFile )
+			CSVFile->GetLastParticleIndexAtTime( t, value );
 
-	*OutValue.GetDest() = value;
-	TimeParam.Advance();
-	OutValue.Advance();
+		*OutValue.GetDest() = value;
+		TimeParam.Advance();
+		OutValue.Advance();
     }
 }
 
@@ -553,65 +658,65 @@ void UNiagaraDataInterfaceHoudiniCSV::GetParticleIndexesToSpawnAtTime( FVectorVM
 
     for (int32 i = 0; i < Context.NumInstances; ++i)
     {
-	float t = TimeParam.Get();
+		float t = TimeParam.Get();
 
-	int32 value = 0;
-	int32 min = 0, max = 0, count = 0;
-	if ( CSVFile )
-	{
-	    if ( !CSVFile->GetLastParticleIndexAtTime( t, value ) )
-	    {
-		// The CSV file doesn't have time informations, so always return all points in the file
-		min = 0;
-		max = value;
-		count = max - min + 1;
-	    }
-	    else
-	    {
-		// The CSV file has time informations
-		// First, detect if we need to reset LastSpawnIndex (after a loop)
-		if ( value < LastSpawnIndex )
-		    LastSpawnIndex = -1;
-
-		if ( value < 0 )
+		int32 value = 0;
+		int32 min = 0, max = 0, count = 0;
+		if ( CSVFile )
 		{
-		    // Nothing to spawn, t is lower than the particle's time
-		    LastSpawnIndex = -1;
+			if ( !CSVFile->GetLastParticleIndexAtTime( t, value ) )
+			{
+				// The CSV file doesn't have time informations, so always return all points in the file
+				min = 0;
+				max = value;
+				count = max - min + 1;
+			}
+			else
+			{
+				// The CSV file has time informations
+				// First, detect if we need to reset LastSpawnIndex (after a loop)
+				if ( value < LastSpawnIndex )
+					LastSpawnIndex = -1;
+
+				if ( value < 0 )
+				{
+					// Nothing to spawn, t is lower than the particle's time
+					LastSpawnIndex = -1;
+				}
+				else
+				{
+					// The last time value in the CSV is lower than t, spawn everything if we didnt already!
+					if ( value >= CSVFile->GetNumberOfPointsInCSV() )
+						value = value - 1;
+
+					if ( value == LastSpawnIndex )
+					{
+						// We've already spawned all the particles for this time
+						min = value;
+						max = value;
+						count = 0;
+					}
+					else
+					{
+						// Found particles to spawn at time t
+						min = LastSpawnIndex + 1;
+						max = value;
+						count = max - min + 1;
+
+						LastSpawnIndex = max;
+					}
+				}
+			}
 		}
-		else
-		{
-		    // The last time value in the CSV is lower than t, spawn everything if we didnt already!
-		    if ( value >= CSVFile->GetNumberOfPointsInCSV() )
-			value = value - 1;
-
-		    if ( value == LastSpawnIndex )
-		    {
-			// We've already spawned all the particles for this time
-			min = value;
-			max = value;
-			count = 0;
-		    }
-		    else
-		    {
-			// Found particles to spawn at time t
-			min = LastSpawnIndex + 1;
-			max = value;
-			count = max - min + 1;
-
-			LastSpawnIndex = max;
-		    }
-		}
-	    }
-	}
 
-	*OutMinValue.GetDest() = min;
-	*OutMaxValue.GetDest() = max;
-	*OutCountValue.GetDest() = count;
+		*OutMinValue.GetDest() = min;
+		*OutMaxValue.GetDest() = max;
+		*OutCountValue.GetDest() = count;
 
-	TimeParam.Advance();
-	OutMinValue.Advance();
-	OutMaxValue.Advance();
-	OutCountValue.Advance();
+		TimeParam.Advance();
+		OutMinValue.Advance();
+		OutMaxValue.Advance();
+		OutCountValue.Advance();
     }
 }
 
@@ -627,28 +732,165 @@ void UNiagaraDataInterfaceHoudiniCSV::GetCSVPositionAndTime(FVectorVMContext& Co
 
     for (int32 i = 0; i < Context.NumInstances; ++i)
     {
-	float N = NParam.Get();
+		float N = NParam.Get();
+
+		float timeValue = 0.0f;
+		FVector posVector = FVector::ZeroVector;
+		if ( CSVFile )
+		{
+			CSVFile->GetCSVTimeValue(N, timeValue);
+			CSVFile->GetCSVPositionValue(N, posVector);
+		}
+
+		*OutPosX.GetDest() = posVector.X;
+		*OutPosY.GetDest() = posVector.Y;
+		*OutPosZ.GetDest() = posVector.Z;
+
+		*OutTime.GetDest() = timeValue;
+
+		NParam.Advance();
+		OutPosX.Advance();
+		OutPosY.Advance();
+		OutPosZ.Advance();
+		OutTime.Advance();
+    }
+}
+
+template<typename NParamType, typename TimeParamType>
+void UNiagaraDataInterfaceHoudiniCSV::GetParticleIndexesAtTime(FVectorVMContext& Context)
+{
+	NParamType NParam(Context);
+	TimeParamType TimeParam(Context);
+
+	FRegisterHandler<int32> OutPrevIndex(Context);
+	FRegisterHandler<int32> OutNextIndex(Context);
+	FRegisterHandler<float> OutWeightValue(Context);
+
+	for (int32 i = 0; i < Context.NumInstances; ++i)
+    {
+		float N = NParam.Get();
+		float time = TimeParam.Get();
+
+		float weight = 0.0f;
+		int32 prevIdx = 0;
+		int32 nextIdx = 0;
+		if ( CSVFile )
+		{
+			CSVFile->GetParticleLineIndexAtTime( N, time, prevIdx, nextIdx, weight );
+		}
+
+		*OutPrevIndex.GetDest() = prevIdx;
+		*OutNextIndex.GetDest() = nextIdx;
+		*OutWeightValue.GetDest() = weight;
+
+		NParam.Advance();
+		TimeParam.Advance();
+		OutPrevIndex.Advance();
+		OutNextIndex.Advance();
+		OutWeightValue.Advance();
+    }
+}
+
+template<typename NParamType, typename TimeParamType>
+void UNiagaraDataInterfaceHoudiniCSV::GetParticlePositionAtTime( FVectorVMContext& Context )
+{
+	NParamType NParam(Context);
+	TimeParamType TimeParam(Context);
+
+	//FRegisterHandler<FVector> OutValue(Context);
+	FRegisterHandler<float> OutPosX(Context);
+	FRegisterHandler<float> OutPosY(Context);
+	FRegisterHandler<float> OutPosZ(Context);
+
+	for (int32 i = 0; i < Context.NumInstances; ++i)
+    {
+		float N = NParam.Get();
+		float time = TimeParam.Get();
+
+		FVector posVector = FVector::ZeroVector;
+		if ( CSVFile )
+		{
+			CSVFile->GetParticlePositionAtTime(N, time, posVector);
+		}		
+
+		//*OutValue.GetDest() = posVector;
+		*OutPosX.GetDest() = posVector.X;
+		*OutPosY.GetDest() = posVector.Y;
+		*OutPosZ.GetDest() = posVector.Z;
+
+		NParam.Advance();
+		TimeParam.Advance();
+		//OutValue.Advance();
+		OutPosX.Advance();
+		OutPosY.Advance();
+		OutPosZ.Advance();
+    }
+}
+
+template<typename NParamType, typename ColParamType, typename TimeParamType>
+void UNiagaraDataInterfaceHoudiniCSV::GetParticleValueAtTime(FVectorVMContext& Context)
+{
+	NParamType NParam(Context);
+	TimeParamType TimeParam(Context);
+	ColParamType ColParam(Context);
+
+	FRegisterHandler<float> OutValue(Context);
 
-	float timeValue = 0.0f;
-	FVector posVector = FVector::ZeroVector;
-	if ( CSVFile )
+	for (int32 i = 0; i < Context.NumInstances; ++i)
 	{
-	    CSVFile->GetCSVTimeValue(N, timeValue);
-	    CSVFile->GetCSVPositionValue(N, posVector);
+		float N = NParam.Get();
+		float time = TimeParam.Get();
+		float Col = ColParam.Get();
+
+		float Value = 0.0f;
+		if ( CSVFile )
+		{
+			CSVFile->GetParticleValueAtTime(N, Col, time, Value);
+		}
+
+		*OutValue.GetDest() = Value;
+
+		NParam.Advance();
+		TimeParam.Advance();
+		ColParam.Advance();
+		OutValue.Advance();
 	}
+}
+
+template<typename NParamType, typename ColParamType, typename TimeParamType>
+void UNiagaraDataInterfaceHoudiniCSV::GetParticleVectorValueAtTime(FVectorVMContext& Context)
+{
+	NParamType NParam(Context);
+	TimeParamType TimeParam(Context);
+	ColParamType ColParam(Context);
 
-	*OutPosX.GetDest() = posVector.X;
-	*OutPosY.GetDest() = posVector.Y;
-	*OutPosZ.GetDest() = posVector.Z;
+	FRegisterHandler<float> OutPosX(Context);
+	FRegisterHandler<float> OutPosY(Context);
+	FRegisterHandler<float> OutPosZ(Context);
 
-	*OutTime.GetDest() = timeValue;
+	for (int32 i = 0; i < Context.NumInstances; ++i)
+	{
+		float N = NParam.Get();
+		float time = TimeParam.Get();
+		float Col = ColParam.Get();
 
-	NParam.Advance();
-	OutPosX.Advance();
-	OutPosY.Advance();
-	OutPosZ.Advance();
-	OutTime.Advance();
-    }
+		FVector posVector = FVector::ZeroVector;
+		if ( CSVFile )
+		{
+			CSVFile->GetParticleVectorValueAtTime(N, Col, time, posVector, true, true);
+		}
+
+		*OutPosX.GetDest() = posVector.X;
+		*OutPosY.GetDest() = posVector.Y;
+		*OutPosZ.GetDest() = posVector.Z;
+
+		NParam.Advance();
+		TimeParam.Advance();
+		ColParam.Advance();
+		OutPosX.Advance();
+		OutPosY.Advance();
+		OutPosZ.Advance();
+	}
 }
 
 void UNiagaraDataInterfaceHoudiniCSV::GetNumberOfPointsInCSV(FVectorVMContext& Context)
diff --git a/Engine/Plugins/FX/HoudiniNiagara/Source/HoudiniNiagara/Public/HoudiniCSV.h b/Engine/Plugins/FX/HoudiniNiagara/Source/HoudiniNiagara/Public/HoudiniCSV.h
--- a/Engine/Plugins/FX/HoudiniNiagara/Source/HoudiniNiagara/Public/HoudiniCSV.h
+++ b/Engine/Plugins/FX/HoudiniNiagara/Source/HoudiniNiagara/Public/HoudiniCSV.h
@@ -1,5 +1,5 @@
 /*
-* Copyright (c) <2017> Side Effects Software Inc.
+* Copyright (c) <2018> Side Effects Software Inc.
 *
 * Permission is hereby granted, free of charge, to any person obtaining a copy
 * of this software and associated documentation files (the "Software"), to deal
@@ -48,13 +48,16 @@ class HOUDININIAGARA_API UHoudiniCSV : public UObject
 	// Returns the number of points found in the CSV file
 	int32 GetNumberOfPointsInCSV();
 
+	// Returns the number of points found in the CSV file
+	int32 GetNumberOfLinesInCSV();
+
 	// Returns the column index for a given string
 	bool GetColumnIndexFromString(const FString& ColumnTitle, int32& ColumnIndex);
 
 	// Returns the float value at a given point in the CSV file
 	bool GetCSVFloatValue( const int32& lineIndex, const int32& colIndex, float& value );
 	// Returns the float value at a given point in the CSV file
-	bool GetCSVFloatValue(const int32& lineIndex, const FString& ColumnTitle, float& value);
+	bool GetCSVFloatValue( const int32& lineIndex, const FString& ColumnTitle, float& value );
 	// Returns the float value at a given point in the CSV file
 	bool GetCSVStringValue( const int32& lineIndex, const int32& colIndex, FString& value );
 	// Returns the string value at a given point in the CSV file
@@ -74,11 +77,18 @@ class HOUDININIAGARA_API UHoudiniCSV : public UObject
 	// -1 will be returned if no particles have been spawned ( t is smaller than the first particle time )
 	// NumberOfLines will be returned if all particles in the CSV have been spawned ( t is higher than the last particle time )
 	bool GetLastParticleIndexAtTime( const float& time, int32& lastIndex );
+	// Returns the previous and next indexes for reading the values of a specified particle at a given time
+	bool GetParticleLineIndexAtTime(const int32& ParticleID, const float& desiredTime, int32& PrevIndex, int32& NextIndex, float& PrevWeight);
+	// Returns the value for a particle at a given time value (linearly interpolated) 
+	bool GetParticleValueAtTime(const int32& ParticleID, const int32& ColumnIndex, const float& desiredTime, float& Value);
+	// Returns the Vector Value for a given particle at a given time value (linearly interpolated) 
+	bool GetParticleVectorValueAtTime(const int32& ParticleID, const int32& ColumnIndex, const float& desiredTime, FVector& Vector, const bool& DoSwap, const bool& DoScale);
+	// Returns the Position Value for a given particle at a given time value (linearly interpolated) 
+	bool GetParticlePositionAtTime(const int32& ParticleID, const float& desiredTime, FVector& Vector);
 
 	//-----------------------------------------------------------------------------------------
 	//  MEMBER VARIABLES
 	//-----------------------------------------------------------------------------------------
-	//UPROPERTY(EditAnywhere, Category = "My Object Properties")
 	UPROPERTY( VisibleAnywhere, Category = "Houdini CSV File Properties" )
 	FString FileName;
 
@@ -90,12 +100,16 @@ class HOUDININIAGARA_API UHoudiniCSV : public UObject
 	UPROPERTY( VisibleAnywhere, Category = "Houdini CSV File Properties" )
 	int32 NumberOfColumns;
 
+	// The number of particles found in the CSV file
+	UPROPERTY(VisibleAnywhere, Category = "Houdini CSV File Properties")
+	int32 NumberOfParticles;
+
 	// The tokenized title raw, describing the content of each column
 	UPROPERTY( VisibleAnywhere, Category = "Houdini CSV File Properties" )
 	TArray<FString> TitleRowArray;
 
 #if WITH_EDITORONLY_DATA
-	/** Importing data and options used for this mesh */
+	/** Importing data and options used for this asset */
 	UPROPERTY(EditAnywhere, Instanced, Category = ImportSettings)
 	class UAssetImportData* AssetImportData;
 
@@ -112,18 +126,36 @@ class HOUDININIAGARA_API UHoudiniCSV : public UObject
 	UPROPERTY()
 	TArray<float> FloatCSVData;
 
-	// Array containing the Position buffer
+	// Array containing the different time values for each particles in the file
+	TArray<float> TimeValues;
+
+	// Array containing the spawn times for each particles in the file
+	TArray<float> SpawnTimes;
+
+	// Array containing all the life values for each particles in the file 
+	TArray<float> LifeValues;	
+
+	// Index of the Position values in the buffer
 	UPROPERTY()	    
 	int32 PositionColumnIndex;
-	//TArray<float> PositionData;
 
-	// Array containing the Normal buffer
+	// Index of the Normal values in the buffer
 	UPROPERTY()
 	int32 NormalColumnIndex;
-	//TArray<float> NormalData;
 
-	// Array containing the time buffer
+	// Index of the time values in the buffer
 	UPROPERTY()
 	int32 TimeColumnIndex;
-	//TArray<float> TimeData;
+
+	// Index of the particle id values in the buffer
+	UPROPERTY()
+	int32 IDColumnIndex;
+
+	// Index of the alive values in the buffer
+	UPROPERTY()
+	int32 AliveColumnIndex;
+
+	// Index of the life values in the buffer
+	UPROPERTY()
+	int32 LifeColumnIndex;
 };
\ No newline at end of file
diff --git a/Engine/Plugins/FX/HoudiniNiagara/Source/HoudiniNiagara/Public/HoudiniNiagara.h b/Engine/Plugins/FX/HoudiniNiagara/Source/HoudiniNiagara/Public/HoudiniNiagara.h
--- a/Engine/Plugins/FX/HoudiniNiagara/Source/HoudiniNiagara/Public/HoudiniNiagara.h
+++ b/Engine/Plugins/FX/HoudiniNiagara/Source/HoudiniNiagara/Public/HoudiniNiagara.h
@@ -1,5 +1,5 @@
 /*
-* Copyright (c) <2017> Side Effects Software Inc.
+* Copyright (c) <2018> Side Effects Software Inc.
 *
 * Permission is hereby granted, free of charge, to any person obtaining a copy
 * of this software and associated documentation files (the "Software"), to deal
diff --git a/Engine/Plugins/FX/HoudiniNiagara/Source/HoudiniNiagaraEditor/HoudiniNiagaraEditor.Build.cs b/Engine/Plugins/FX/HoudiniNiagara/Source/HoudiniNiagaraEditor/HoudiniNiagaraEditor.Build.cs
--- a/Engine/Plugins/FX/HoudiniNiagara/Source/HoudiniNiagaraEditor/HoudiniNiagaraEditor.Build.cs
+++ b/Engine/Plugins/FX/HoudiniNiagara/Source/HoudiniNiagaraEditor/HoudiniNiagaraEditor.Build.cs
@@ -1,5 +1,5 @@
 /*
-* Copyright (c) <2017> Side Effects Software Inc.
+* Copyright (c) <2018> Side Effects Software Inc.
 *
 * Permission is hereby granted, free of charge, to any person obtaining a copy
 * of this software and associated documentation files (the "Software"), to deal
@@ -52,7 +52,7 @@ public class HoudiniNiagaraEditor : ModuleRules
                 "ShaderCore",
                 "HoudiniNiagara",
                 "UnrealEd",
-                                "EditorStyle",
+				"EditorStyle",
             }
         );
 			
diff --git a/Engine/Plugins/FX/HoudiniNiagara/Source/HoudiniNiagaraEditor/Private/HoudiniCSVAssetActions.cpp b/Engine/Plugins/FX/HoudiniNiagara/Source/HoudiniNiagaraEditor/Private/HoudiniCSVAssetActions.cpp
--- a/Engine/Plugins/FX/HoudiniNiagara/Source/HoudiniNiagaraEditor/Private/HoudiniCSVAssetActions.cpp
+++ b/Engine/Plugins/FX/HoudiniNiagara/Source/HoudiniNiagaraEditor/Private/HoudiniCSVAssetActions.cpp
@@ -1,5 +1,5 @@
 /*
-* Copyright (c) <2017> Side Effects Software Inc.
+* Copyright (c) <2018> Side Effects Software Inc.
 *
 * Permission is hereby granted, free of charge, to any person obtaining a copy
 * of this software and associated documentation files (the "Software"), to deal
@@ -49,23 +49,23 @@ void FHoudiniCSVAssetActions::GetActions(const TArray<UObject*>& InObjects, FMen
     auto HoudiniCSVAssets = GetTypedWeakObjectPtrs<UHoudiniCSV>(InObjects);
 
     MenuBuilder.AddMenuEntry(
-	LOCTEXT("ReimportHoudiniCSVLabel", "Reimport"),
-	LOCTEXT("ReimportHoudiniCSVTooltip", "Reimport the selected Houdini CSV file(s)."),
-	FSlateIcon(FEditorStyle::GetStyleSetName(), "ContentBrowser.AssetActions.ReimportAsset"),
-	FUIAction(
-	    FExecuteAction::CreateSP(this, &FHoudiniCSVAssetActions::ExecuteReimport, HoudiniCSVAssets),
-	    FCanExecuteAction::CreateSP(this, &FHoudiniCSVAssetActions::CanExecuteReimport, HoudiniCSVAssets)
-	)
+		LOCTEXT("ReimportHoudiniCSVLabel", "Reimport"),
+		LOCTEXT("ReimportHoudiniCSVTooltip", "Reimport the selected Houdini CSV file(s)."),
+		FSlateIcon(FEditorStyle::GetStyleSetName(), "ContentBrowser.AssetActions.ReimportAsset"),
+		FUIAction(
+			FExecuteAction::CreateSP(this, &FHoudiniCSVAssetActions::ExecuteReimport, HoudiniCSVAssets),
+			FCanExecuteAction::CreateSP(this, &FHoudiniCSVAssetActions::CanExecuteReimport, HoudiniCSVAssets)
+		)
     );
 
     MenuBuilder.AddMenuEntry(
-	LOCTEXT("OpenHoudiniCSVLabel", "Open in Text Editor"),
-	LOCTEXT("OpenHoudiniCSVTooltip", "Open the selected Houdini CSV file(s) in a Text Editor."),
-	FSlateIcon(FEditorStyle::GetStyleSetName(), "ContentBrowser.AssetActions.OpenInExternalEditor"),
-	FUIAction(
-	    FExecuteAction::CreateSP(this, &FHoudiniCSVAssetActions::ExecuteOpenInEditor, HoudiniCSVAssets),
-	    FCanExecuteAction::CreateSP(this, &FHoudiniCSVAssetActions::CanExecuteOpenInEditor, HoudiniCSVAssets)
-	)
+		LOCTEXT("OpenHoudiniCSVLabel", "Open in Text Editor"),
+		LOCTEXT("OpenHoudiniCSVTooltip", "Open the selected Houdini CSV file(s) in a Text Editor."),
+		FSlateIcon(FEditorStyle::GetStyleSetName(), "ContentBrowser.AssetActions.OpenInExternalEditor"),
+		FUIAction(
+			FExecuteAction::CreateSP(this, &FHoudiniCSVAssetActions::ExecuteOpenInEditor, HoudiniCSVAssets),
+			FCanExecuteAction::CreateSP(this, &FHoudiniCSVAssetActions::CanExecuteOpenInEditor, HoudiniCSVAssets)
+		)
     );
 }
 
@@ -127,9 +127,9 @@ bool FHoudiniCSVAssetActions::CanExecuteReimport(const TArray<TWeakObjectPtr<UHo
 {
     for (auto ObjIt = Objects.CreateConstIterator(); ObjIt; ++ObjIt)
     {
-	auto Object = (*ObjIt).Get();
-	if (!Object)
-	    return false;
+		auto Object = (*ObjIt).Get();
+		if (!Object)
+			return false;
     }
 
     return true;
@@ -139,12 +139,12 @@ void FHoudiniCSVAssetActions::ExecuteReimport(const TArray<TWeakObjectPtr<UHoudi
 {
     for ( auto ObjIt = Objects.CreateConstIterator(); ObjIt; ++ObjIt )
     {
-	auto Object = (*ObjIt).Get();
-	if ( Object )
-	{
-	    // Fonts fail to reimport if they ask for a new file if missing
-	    FReimportManager::Instance()->Reimport(Object, true);
-	}
+		auto Object = (*ObjIt).Get();
+		if ( Object )
+		{
+			// Fonts fail to reimport if they ask for a new file if missing
+			FReimportManager::Instance()->Reimport(Object, true);
+		}
     }
 }
 
@@ -154,9 +154,9 @@ FHoudiniCSVAssetActions::CanExecuteOpenInEditor(const TArray<TWeakObjectPtr<UHou
 {
     for (auto ObjIt = Objects.CreateConstIterator(); ObjIt; ++ObjIt)
     {
-	auto Object = ( *ObjIt ).Get();
-	if ( !Object )
-	    return false;
+		auto Object = ( *ObjIt ).Get();
+		if ( !Object )
+			return false;
     }
 
     return true;
@@ -167,12 +167,12 @@ FHoudiniCSVAssetActions::ExecuteOpenInEditor(const TArray<TWeakObjectPtr<UHoudin
 {
     for ( auto ObjIt = Objects.CreateConstIterator(); ObjIt; ++ObjIt )
     {
-	auto Object = ( *ObjIt ).Get();
-	if ( Object )
-	{
-	    // Fonts fail to reimport if they ask for a new file if missing
-	    FPlatformProcess::LaunchFileInDefaultExternalApplication(*(Object->FileName), NULL, ELaunchVerb::Open);
-	}
+		auto Object = ( *ObjIt ).Get();
+		if ( Object )
+		{
+			// Fonts fail to reimport if they ask for a new file if missing
+			FPlatformProcess::LaunchFileInDefaultExternalApplication(*(Object->FileName), NULL, ELaunchVerb::Open);
+		}
     }
 }
 
diff --git a/Engine/Plugins/FX/HoudiniNiagara/Source/HoudiniNiagaraEditor/Private/HoudiniCSVFactory.cpp b/Engine/Plugins/FX/HoudiniNiagara/Source/HoudiniNiagaraEditor/Private/HoudiniCSVFactory.cpp
--- a/Engine/Plugins/FX/HoudiniNiagara/Source/HoudiniNiagaraEditor/Private/HoudiniCSVFactory.cpp
+++ b/Engine/Plugins/FX/HoudiniNiagara/Source/HoudiniNiagaraEditor/Private/HoudiniCSVFactory.cpp
@@ -1,5 +1,5 @@
 /*
-* Copyright (c) <2017> Side Effects Software Inc.
+* Copyright (c) <2018> Side Effects Software Inc.
 *
 * Permission is hereby granted, free of charge, to any person obtaining a copy
 * of this software and associated documentation files (the "Software"), to deal
@@ -72,17 +72,17 @@ UHoudiniCSVFactory::FactoryCreateFile(UClass* InClass, UObject* InParent, FName
 
     UHoudiniCSV* NewHoudiniCSVObject = NewObject<UHoudiniCSV>(InParent, InClass, InName, Flags);
     if ( !NewHoudiniCSVObject )
-	return nullptr;
+		return nullptr;
 
     if ( !NewHoudiniCSVObject->UpdateFromFile( Filename ) )
-	return nullptr;   
+		return nullptr;   
 
     // Create reimport information.
     UAssetImportData * AssetImportData = NewHoudiniCSVObject->AssetImportData;
     if (!AssetImportData)
     {
-	AssetImportData = NewObject< UAssetImportData >(NewHoudiniCSVObject, UAssetImportData::StaticClass());
-	NewHoudiniCSVObject->AssetImportData = AssetImportData;
+		AssetImportData = NewObject< UAssetImportData >(NewHoudiniCSVObject, UAssetImportData::StaticClass());
+		NewHoudiniCSVObject->AssetImportData = AssetImportData;
     }
 
     //NewHoudiniCSVObject->AssetImportData->Update(Filename);
@@ -112,7 +112,7 @@ UHoudiniCSVFactory::FactoryCanImport(const FString& Filename)
 
     if ( Extension == TEXT("csv") || Extension == TEXT("hcsv") )
     {
-	return true;
+		return true;
     }
     return false;
 }
@@ -123,15 +123,15 @@ UHoudiniCSVFactory::CanReimport(UObject* Obj, TArray<FString>& OutFilenames)
     UHoudiniCSV* HoudiniCSV = Cast<UHoudiniCSV>(Obj);
     if (HoudiniCSV)
     {
-	if (HoudiniCSV->AssetImportData)
-	{
-	    HoudiniCSV->AssetImportData->ExtractFilenames(OutFilenames);
-	}
-	else
-	{
-	    OutFilenames.Add(TEXT(""));
-	}
-	return true;
+		if (HoudiniCSV->AssetImportData)
+		{
+			HoudiniCSV->AssetImportData->ExtractFilenames(OutFilenames);
+		}
+		else
+		{
+			OutFilenames.Add(TEXT(""));
+		}
+		return true;
     }
     return false;
 }
@@ -139,12 +139,12 @@ UHoudiniCSVFactory::CanReimport(UObject* Obj, TArray<FString>& OutFilenames)
 void 
 UHoudiniCSVFactory::SetReimportPaths(UObject* Obj, const TArray<FString>& NewReimportPaths)
 {
-    UHoudiniCSV* HoudiniCSV = Cast<UHoudiniCSV>(Obj);
-    if (HoudiniCSV 
-	&& HoudiniCSV->AssetImportData
-	&&  ensure(NewReimportPaths.Num() == 1))
+    UHoudiniCSV* HoudiniCSV = Cast<UHoudiniCSV>( Obj );
+    if ( HoudiniCSV 
+		&& HoudiniCSV->AssetImportData
+		&&  ensure( NewReimportPaths.Num() == 1 ) )
     {
-	HoudiniCSV->AssetImportData->UpdateFilenameOnly(NewReimportPaths[0]);
+		HoudiniCSV->AssetImportData->UpdateFilenameOnly(NewReimportPaths[0]);
     }
 }
 
@@ -153,38 +153,38 @@ UHoudiniCSVFactory::Reimport(UObject* Obj)
 {
     UHoudiniCSV* HoudiniCSV = Cast<UHoudiniCSV>(Obj);
     if (!HoudiniCSV || !HoudiniCSV->AssetImportData )
-	return EReimportResult::Failed;
+		return EReimportResult::Failed;
 
     const FString FilePath = HoudiniCSV->AssetImportData->GetFirstFilename();
     if (!FilePath.Len())
-	return EReimportResult::Failed;
+		return EReimportResult::Failed;
 
     UE_LOG(LogHoudiniNiagaraEditor, Log, TEXT("Performing atomic reimport of [%s]"), *FilePath);
 
     // Ensure that the file provided by the path exists
     if (IFileManager::Get().FileSize(*FilePath) == INDEX_NONE)
     {
-	UE_LOG(LogHoudiniNiagaraEditor, Warning, TEXT("Cannot reimport: source file cannot be found."));
-	return EReimportResult::Failed;
+		UE_LOG(LogHoudiniNiagaraEditor, Warning, TEXT("Cannot reimport: source file cannot be found."));
+		return EReimportResult::Failed;
     }
 
     bool OutCanceled = false;
     if (ImportObject(HoudiniCSV->GetClass(), HoudiniCSV->GetOuter(), *HoudiniCSV->GetName(), RF_Public | RF_Standalone, FilePath, nullptr, OutCanceled) != nullptr)
     {
-	UE_LOG(LogHoudiniNiagaraEditor, Log, TEXT("Imported successfully"));
-	// Try to find the outer package so we can dirty it up
-	if (HoudiniCSV->GetOuter())
-	    HoudiniCSV->GetOuter()->MarkPackageDirty();
-	else
-	    HoudiniCSV->MarkPackageDirty();
+		UE_LOG(LogHoudiniNiagaraEditor, Log, TEXT("Imported successfully"));
+		// Try to find the outer package so we can dirty it up
+		if (HoudiniCSV->GetOuter())
+			HoudiniCSV->GetOuter()->MarkPackageDirty();
+		else
+			HoudiniCSV->MarkPackageDirty();
     }
     else if (OutCanceled)
     {
-	UE_LOG(LogHoudiniNiagaraEditor, Warning, TEXT("-- import canceled"));
+		UE_LOG(LogHoudiniNiagaraEditor, Warning, TEXT("-- import canceled"));
     }
     else
     {
-	UE_LOG(LogHoudiniNiagaraEditor, Warning, TEXT("-- import failed"));
+		UE_LOG(LogHoudiniNiagaraEditor, Warning, TEXT("-- import failed"));
     }
 
     return EReimportResult::Succeeded;
diff --git a/Engine/Plugins/FX/HoudiniNiagara/Source/HoudiniNiagaraEditor/Private/HoudiniNiagaraEditor.cpp b/Engine/Plugins/FX/HoudiniNiagara/Source/HoudiniNiagaraEditor/Private/HoudiniNiagaraEditor.cpp
--- a/Engine/Plugins/FX/HoudiniNiagara/Source/HoudiniNiagaraEditor/Private/HoudiniNiagaraEditor.cpp
+++ b/Engine/Plugins/FX/HoudiniNiagara/Source/HoudiniNiagaraEditor/Private/HoudiniNiagaraEditor.cpp
@@ -1,5 +1,5 @@
 /*
-* Copyright (c) <2017> Side Effects Software Inc.
+* Copyright (c) <2018> Side Effects Software Inc.
 *
 * Permission is hereby granted, free of charge, to any person obtaining a copy
 * of this software and associated documentation files (the "Software"), to deal
@@ -29,36 +29,28 @@
 
 void FHoudiniNiagaraEditorModule::StartupModule()
 {
-    // This code will execute after your module is loaded into memory; the exact timing is specified in the .uplugin file per-module
-
     // Register the Houdini CSV Type Actions
     IAssetTools& AssetTools = FModuleManager::LoadModuleChecked< FAssetToolsModule >("AssetTools").Get();
 
-    TSharedRef< IAssetTypeActions > HCSVAction = MakeShareable(new FHoudiniCSVAssetActions());
-    AssetTools.RegisterAssetTypeActions(HCSVAction);
-    AssetTypeActions.Add(HCSVAction);
+    TSharedRef< IAssetTypeActions > HCSVAction = MakeShareable( new FHoudiniCSVAssetActions() );
+    AssetTools.RegisterAssetTypeActions( HCSVAction );
+    AssetTypeActions.Add( HCSVAction );
 }
 
 void FHoudiniNiagaraEditorModule::ShutdownModule()
 {
-    // This function may be called during shutdown to clean up your module.  For modules that support dynamic reloading,
-    // we call this function before unloading the module.
-
-
     // Unregister asset type actions we have previously registered.
-    if (FModuleManager::Get().IsModuleLoaded("AssetTools"))
+    if ( FModuleManager::Get().IsModuleLoaded("AssetTools") )
     {
-	IAssetTools & AssetTools = FModuleManager::GetModuleChecked< FAssetToolsModule >("AssetTools").Get();
+		IAssetTools & AssetTools = FModuleManager::GetModuleChecked< FAssetToolsModule >("AssetTools").Get();
 
-	for (int32 Index = 0; Index < AssetTypeActions.Num(); ++Index)
-	    AssetTools.UnregisterAssetTypeActions(AssetTypeActions[Index].ToSharedRef());
+		for ( int32 Index = 0; Index < AssetTypeActions.Num(); ++Index )
+			AssetTools.UnregisterAssetTypeActions( AssetTypeActions[Index].ToSharedRef() );
 
-	AssetTypeActions.Empty();
+		AssetTypeActions.Empty();
     }
 }
 
-
-
 #undef LOCTEXT_NAMESPACE
 	
 IMPLEMENT_MODULE(FHoudiniNiagaraEditorModule, HoudiniNiagara)
diff --git a/Engine/Plugins/FX/HoudiniNiagara/Source/HoudiniNiagaraEditor/Public/HoudiniCSVAssetActions.h b/Engine/Plugins/FX/HoudiniNiagara/Source/HoudiniNiagaraEditor/Public/HoudiniCSVAssetActions.h
--- a/Engine/Plugins/FX/HoudiniNiagara/Source/HoudiniNiagaraEditor/Public/HoudiniCSVAssetActions.h
+++ b/Engine/Plugins/FX/HoudiniNiagara/Source/HoudiniNiagaraEditor/Public/HoudiniCSVAssetActions.h
@@ -1,5 +1,5 @@
 /*
-* Copyright (c) <2017> Side Effects Software Inc.
+* Copyright (c) <2018> Side Effects Software Inc.
 *
 * Permission is hereby granted, free of charge, to any person obtaining a copy
 * of this software and associated documentation files (the "Software"), to deal
diff --git a/Engine/Plugins/FX/HoudiniNiagara/Source/HoudiniNiagaraEditor/Public/HoudiniCSVFactory.h b/Engine/Plugins/FX/HoudiniNiagara/Source/HoudiniNiagaraEditor/Public/HoudiniCSVFactory.h
--- a/Engine/Plugins/FX/HoudiniNiagara/Source/HoudiniNiagaraEditor/Public/HoudiniCSVFactory.h
+++ b/Engine/Plugins/FX/HoudiniNiagara/Source/HoudiniNiagaraEditor/Public/HoudiniCSVFactory.h
@@ -1,5 +1,5 @@
 /*
-* Copyright (c) <2017> Side Effects Software Inc.
+* Copyright (c) <2018> Side Effects Software Inc.
 *
 * Permission is hereby granted, free of charge, to any person obtaining a copy
 * of this software and associated documentation files (the "Software"), to deal
@@ -32,7 +32,7 @@
 DECLARE_LOG_CATEGORY_EXTERN(LogHoudiniNiagaraEditor, All, All);
 
 /** A factory for HoudiniCSV assets. */
-UCLASS(MinimalAPI/*UNREALED_API*/)
+UCLASS( MinimalAPI )
 class UHoudiniCSVFactory : public UFactory, public FReimportHandler
 {
 	GENERATED_UCLASS_BODY()
diff --git a/Engine/Plugins/FX/HoudiniNiagara/Source/HoudiniNiagaraEditor/Public/HoudiniNiagaraEditor.h b/Engine/Plugins/FX/HoudiniNiagara/Source/HoudiniNiagaraEditor/Public/HoudiniNiagaraEditor.h
--- a/Engine/Plugins/FX/HoudiniNiagara/Source/HoudiniNiagaraEditor/Public/HoudiniNiagaraEditor.h
+++ b/Engine/Plugins/FX/HoudiniNiagara/Source/HoudiniNiagaraEditor/Public/HoudiniNiagaraEditor.h
@@ -1,5 +1,5 @@
 /*
-* Copyright (c) <2017> Side Effects Software Inc.
+* Copyright (c) <2018> Side Effects Software Inc.
 *
 * Permission is hereby granted, free of charge, to any person obtaining a copy
 * of this software and associated documentation files (the "Software"), to deal
