Niagara - Houdini:

Houdini Niagara Data Interface:
- Since we can now update particles over time, renamed/modified most of the functions to make a clear distinction between row indexes (row) and particle ids (N)
- Replaced GetNumberOfPointsInCSV by GetNumberOfRowsInCSV and GetNumberOfParticlesInCSV
- Renamed GetParticleIndexesAtTime to GetRowIndexesForParticleAtTime and GetLastParticleIndexAtTime to GetLastRowIndexAtTime
- Fixed some DI Functions that were using floats for input parameter instead of using integers.

diff --git a/Engine/Plugins/FX/HoudiniNiagara/Source/HoudiniNiagara/Classes/NiagaraDataInterfaceHoudiniCSV.h b/Engine/Plugins/FX/HoudiniNiagara/Source/HoudiniNiagara/Classes/NiagaraDataInterfaceHoudiniCSV.h
--- a/Engine/Plugins/FX/HoudiniNiagara/Source/HoudiniNiagara/Classes/NiagaraDataInterfaceHoudiniCSV.h
+++ b/Engine/Plugins/FX/HoudiniNiagara/Source/HoudiniNiagara/Classes/NiagaraDataInterfaceHoudiniCSV.h
@@ -75,27 +75,30 @@ public:
 	void GetCSVVectorValue(FVectorVMContext& Context);
 
 	// Returns the positions for a given point in the CSV file
-	template<typename NParamType>
+	template<typename RowParamType>
 	void GetCSVPosition(FVectorVMContext& Context);
 
 	// Returns the normals for a given point in the CSV file
-	template<typename NParamType>
+	template<typename RowParamType>
 	void GetCSVNormal(FVectorVMContext& Context);
 
 	// Returns the time for a given point in the CSV file
-	template<typename NParamType>
+	template<typename RowParamType>
 	void GetCSVTime(FVectorVMContext& Context);
 
 	// Returns the position and time for a given point in the CSV file
-	template<typename NParamType>
+	template<typename RowParamType>
 	void GetCSVPositionAndTime(FVectorVMContext& Context);
 
+	// Returns the number of rows found in the CSV file
+	void GetNumberOfRowsInCSV(FVectorVMContext& Context);
+
 	// Returns the number of points found in the CSV file
-	void GetNumberOfPointsInCSV(FVectorVMContext& Context);
+	void GetNumberOfParticlesInCSV(FVectorVMContext& Context);
 
 	// Returns the last index of the particles that should be spawned at time t
 	template<typename TimeParamType>
-	void GetLastParticleIndexAtTime(FVectorVMContext& Context);
+	void GetLastRowIndexAtTime(FVectorVMContext& Context);
 
 	// Returns the indexes (min, max) and number of particles that should be spawned at time t
 	template<typename TimeParamType>
@@ -115,7 +118,7 @@ public:
 
 	// Returns the line indexes (previous, next) for reading values for a given particle at a given time
 	template<typename NParamType, typename TimeParamType>
-	void GetParticleIndexesAtTime(FVectorVMContext& Context);
+	void GetRowIndexesForParticleAtTime(FVectorVMContext& Context);
 	
 	//----------------------------------------------------------------------------
 	// GPU / HLSL Functions
diff --git a/Engine/Plugins/FX/HoudiniNiagara/Source/HoudiniNiagara/Private/HoudiniCSV.cpp b/Engine/Plugins/FX/HoudiniNiagara/Source/HoudiniNiagara/Private/HoudiniCSV.cpp
--- a/Engine/Plugins/FX/HoudiniNiagara/Source/HoudiniNiagara/Private/HoudiniCSV.cpp
+++ b/Engine/Plugins/FX/HoudiniNiagara/Source/HoudiniNiagara/Private/HoudiniCSV.cpp
@@ -503,7 +503,7 @@ bool UHoudiniCSV::GetCSVTimeValue( const int32& lineIndex, float& value )
 }
 
 // Returns the number of points found in the CSV file
-int32 UHoudiniCSV::GetNumberOfPointsInCSV()
+int32 UHoudiniCSV::GetNumberOfParticlesInCSV()
 {
 	if ( IDColumnIndex != INDEX_NONE )
 		return NumberOfParticles;
diff --git a/Engine/Plugins/FX/HoudiniNiagara/Source/HoudiniNiagara/Private/NiagaraDataInterfaceHoudiniCSV.cpp b/Engine/Plugins/FX/HoudiniNiagara/Source/HoudiniNiagara/Private/NiagaraDataInterfaceHoudiniCSV.cpp
--- a/Engine/Plugins/FX/HoudiniNiagara/Source/HoudiniNiagara/Private/NiagaraDataInterfaceHoudiniCSV.cpp
+++ b/Engine/Plugins/FX/HoudiniNiagara/Source/HoudiniNiagara/Private/NiagaraDataInterfaceHoudiniCSV.cpp
@@ -114,8 +114,8 @@ void UNiagaraDataInterfaceHoudiniCSV::GetFunctions(TArray<FNiagaraFunctionSignat
 		Sig.bMemberFunction = true;
 		Sig.bRequiresContext = false;
 		Sig.Inputs.Add(FNiagaraVariable(FNiagaraTypeDefinition(GetClass()), TEXT("CSV")));			// CSV in
-		Sig.Inputs.Add(FNiagaraVariable(FNiagaraTypeDefinition::GetFloatDef(), TEXT("Row")));		// Row Index In
-		Sig.Inputs.Add(FNiagaraVariable(FNiagaraTypeDefinition::GetFloatDef(), TEXT("Col")));		// Col Index In
+		Sig.Inputs.Add(FNiagaraVariable(FNiagaraTypeDefinition::GetIntDef(), TEXT("Row")));			// Row Index In
+		Sig.Inputs.Add(FNiagaraVariable(FNiagaraTypeDefinition::GetIntDef(), TEXT("Col")));			// Col Index In
 		Sig.Outputs.Add(FNiagaraVariable(FNiagaraTypeDefinition::GetFloatDef(), TEXT("Value")));	// Float Out
 
 		OutFunctions.Add(Sig);
@@ -129,7 +129,7 @@ void UNiagaraDataInterfaceHoudiniCSV::GetFunctions(TArray<FNiagaraFunctionSignat
 		Sig.bMemberFunction = true;
 		Sig.bRequiresContext = false;
 		Sig.Inputs.Add(FNiagaraVariable(FNiagaraTypeDefinition(GetClass()), TEXT("CSV")));			// CSV in
-		Sig.Inputs.Add(FNiagaraVariable(FNiagaraTypeDefinition::GetFloatDef(), TEXT("Row")));		// Row Index In
+		Sig.Inputs.Add(FNiagaraVariable(FNiagaraTypeDefinition::GetIntDef(), TEXT("Row")));		// Row Index In
 		Sig.Inputs.Add(FNiagaraVariable(FNiagaraTypeDefinition::GetStringDef(), TEXT("ColTitle")));	// Col Title In
 		Sig.Outputs.Add(FNiagaraVariable(FNiagaraTypeDefinition::GetFloatDef(), TEXT("Value")));	// Float Out
 
@@ -144,8 +144,8 @@ void UNiagaraDataInterfaceHoudiniCSV::GetFunctions(TArray<FNiagaraFunctionSignat
 		Sig.bMemberFunction = true;
 		Sig.bRequiresContext = false;
 		Sig.Inputs.Add(FNiagaraVariable(FNiagaraTypeDefinition(GetClass()), TEXT("CSV")));			// CSV in
-		Sig.Inputs.Add(FNiagaraVariable(FNiagaraTypeDefinition::GetFloatDef(), TEXT("Row")));		// Row Index In
-		Sig.Inputs.Add(FNiagaraVariable(FNiagaraTypeDefinition::GetFloatDef(), TEXT("Col")));		// Col Index In
+		Sig.Inputs.Add(FNiagaraVariable(FNiagaraTypeDefinition::GetIntDef(), TEXT("Row")));			// Row Index In
+		Sig.Inputs.Add(FNiagaraVariable(FNiagaraTypeDefinition::GetIntDef(), TEXT("Col")));			// Col Index In
 		Sig.Outputs.Add(FNiagaraVariable(FNiagaraTypeDefinition::GetVec3Def(), TEXT("Value")));		// Vector3 Out
 
 		OutFunctions.Add(Sig);
@@ -158,8 +158,8 @@ void UNiagaraDataInterfaceHoudiniCSV::GetFunctions(TArray<FNiagaraFunctionSignat
 		Sig.bMemberFunction = true;
 		Sig.bRequiresContext = false;
 		Sig.Inputs.Add(FNiagaraVariable(FNiagaraTypeDefinition(GetClass()), TEXT("CSV")));			// CSV in
-		Sig.Inputs.Add(FNiagaraVariable(FNiagaraTypeDefinition::GetFloatDef(), TEXT("N")));			// Point Number In
-		Sig.Outputs.Add(FNiagaraVariable(FNiagaraTypeDefinition::GetVec3Def(), TEXT("Value")));		// Vector3 Out
+		Sig.Inputs.Add(FNiagaraVariable(FNiagaraTypeDefinition::GetIntDef(), TEXT("Row")));			// Row Index In
+		Sig.Outputs.Add(FNiagaraVariable(FNiagaraTypeDefinition::GetVec3Def(), TEXT("Position")));	// Vector3 Out
 
 		OutFunctions.Add(Sig);
     }
@@ -171,7 +171,7 @@ void UNiagaraDataInterfaceHoudiniCSV::GetFunctions(TArray<FNiagaraFunctionSignat
 		Sig.bMemberFunction = true;
 		Sig.bRequiresContext = false;
 		Sig.Inputs.Add(FNiagaraVariable(FNiagaraTypeDefinition(GetClass()), TEXT("CSV")));			// CSV in
-		Sig.Inputs.Add(FNiagaraVariable(FNiagaraTypeDefinition::GetFloatDef(), TEXT("N")));			// Point Number In
+		Sig.Inputs.Add(FNiagaraVariable(FNiagaraTypeDefinition::GetIntDef(), TEXT("Row")));			// Row Index In
 		Sig.Outputs.Add(FNiagaraVariable(FNiagaraTypeDefinition::GetVec3Def(), TEXT("Position")));	// Vector3 Out
 		Sig.Outputs.Add(FNiagaraVariable(FNiagaraTypeDefinition::GetFloatDef(), TEXT("Time")));		// float Out
 
@@ -185,8 +185,8 @@ void UNiagaraDataInterfaceHoudiniCSV::GetFunctions(TArray<FNiagaraFunctionSignat
 		Sig.bMemberFunction = true;
 		Sig.bRequiresContext = false;
 		Sig.Inputs.Add(FNiagaraVariable(FNiagaraTypeDefinition(GetClass()), TEXT("CSV")));			// CSV in
-		Sig.Inputs.Add(FNiagaraVariable(FNiagaraTypeDefinition::GetFloatDef(), TEXT("N")));			// Point Number In
-		Sig.Outputs.Add(FNiagaraVariable(FNiagaraTypeDefinition::GetVec3Def(), TEXT("Value")));		// Vector3 Out
+		Sig.Inputs.Add(FNiagaraVariable(FNiagaraTypeDefinition::GetIntDef(), TEXT("Row")));			// Row Index In
+		Sig.Outputs.Add(FNiagaraVariable(FNiagaraTypeDefinition::GetVec3Def(), TEXT("Normal")));	// Vector3 Out
 
 		OutFunctions.Add(Sig);
     }
@@ -198,33 +198,45 @@ void UNiagaraDataInterfaceHoudiniCSV::GetFunctions(TArray<FNiagaraFunctionSignat
 		Sig.bMemberFunction = true;
 		Sig.bRequiresContext = false;
 		Sig.Inputs.Add(FNiagaraVariable(FNiagaraTypeDefinition(GetClass()), TEXT("CSV")));			// CSV in
-		Sig.Inputs.Add(FNiagaraVariable(FNiagaraTypeDefinition::GetFloatDef(), TEXT("N")));			// Point Number In
-		Sig.Outputs.Add(FNiagaraVariable(FNiagaraTypeDefinition::GetFloatDef(), TEXT("Value")));	// Float Out
+		Sig.Inputs.Add(FNiagaraVariable(FNiagaraTypeDefinition::GetIntDef(), TEXT("Row")));			// Row Index In
+		Sig.Outputs.Add(FNiagaraVariable(FNiagaraTypeDefinition::GetFloatDef(), TEXT("Time")));		// Float Out
 
 		OutFunctions.Add(Sig);
     }
 
     {
-		// GetNumberOfPointsInCSV
+		// GetNumberOfParticlesInCSV
 		FNiagaraFunctionSignature Sig;
-		Sig.Name = TEXT("GetNumberOfPointsInCSV");
+		Sig.Name = TEXT("GetNumberOfParticlesInCSV");
 		Sig.bMemberFunction = true;
 		Sig.bRequiresContext = false;
 		Sig.Inputs.Add(FNiagaraVariable(FNiagaraTypeDefinition(GetClass()), TEXT("CSV")));					// CSV in
-		Sig.Outputs.Add(FNiagaraVariable(FNiagaraTypeDefinition::GetIntDef(), TEXT("NumberOfPoints")));     // Int Out
+		Sig.Outputs.Add(FNiagaraVariable(FNiagaraTypeDefinition::GetIntDef(), TEXT("NumberOfParticles")));  // Int Out
 
 		OutFunctions.Add(Sig);
     }
 
+	{
+		// GetNumberOfRowsInCSV
+		FNiagaraFunctionSignature Sig;
+		Sig.Name = TEXT("GetNumberOfRowsInCSV");
+		Sig.bMemberFunction = true;
+		Sig.bRequiresContext = false;
+		Sig.Inputs.Add(FNiagaraVariable(FNiagaraTypeDefinition(GetClass()), TEXT("CSV")));					// CSV in
+		Sig.Outputs.Add(FNiagaraVariable(FNiagaraTypeDefinition::GetIntDef(), TEXT("NumberOfRows")));		// Int Out
+
+		OutFunctions.Add(Sig);
+	}
+
     {
-		// GetLastParticleIndexAtTime
+		// GetLastRowIndexAtTime
 		FNiagaraFunctionSignature Sig;
-		Sig.Name = TEXT("GetLastParticleIndexAtTime");
+		Sig.Name = TEXT("GetLastRowIndexAtTime");
 		Sig.bMemberFunction = true;
 		Sig.bRequiresContext = false;
-		Sig.Inputs.Add(FNiagaraVariable(FNiagaraTypeDefinition(GetClass()), TEXT("CSV")));				 // CSV in
-		Sig.Inputs.Add(FNiagaraVariable(FNiagaraTypeDefinition::GetFloatDef(), TEXT("Time")));		    // Time in
-		Sig.Outputs.Add(FNiagaraVariable(FNiagaraTypeDefinition::GetIntDef(), TEXT("LastIndex")));	    // Int Out
+		Sig.Inputs.Add(FNiagaraVariable(FNiagaraTypeDefinition(GetClass()), TEXT("CSV")));					// CSV in
+		Sig.Inputs.Add(FNiagaraVariable(FNiagaraTypeDefinition::GetFloatDef(), TEXT("Time")));				// Time in
+		Sig.Outputs.Add(FNiagaraVariable(FNiagaraTypeDefinition::GetIntDef(), TEXT("LastRowIndex")));	    // Int Out
 
 		OutFunctions.Add(Sig);
     }
@@ -245,17 +257,17 @@ void UNiagaraDataInterfaceHoudiniCSV::GetFunctions(TArray<FNiagaraFunctionSignat
     }
 
 	{
-		// GetParticleIndexesAtTime
+		// GetRowIndexesForParticleAtTime
 		FNiagaraFunctionSignature Sig;
-		Sig.Name = TEXT("GetParticleIndexesAtTime");
+		Sig.Name = TEXT("GetRowIndexesForParticleAtTime");
 		Sig.bMemberFunction = true;
 		Sig.bRequiresContext = false;
-		Sig.Inputs.Add(FNiagaraVariable(FNiagaraTypeDefinition(GetClass()), TEXT("CSV")));				// CSV in
-		Sig.Inputs.Add(FNiagaraVariable(FNiagaraTypeDefinition::GetFloatDef(), TEXT("N")));				// Point Number In
-		Sig.Inputs.Add(FNiagaraVariable(FNiagaraTypeDefinition::GetFloatDef(), TEXT("Time")));		    // Time in
-		Sig.Outputs.Add(FNiagaraVariable(FNiagaraTypeDefinition::GetIntDef(), TEXT("PreviousLineIndex"))); // Int Out
-		Sig.Outputs.Add(FNiagaraVariable(FNiagaraTypeDefinition::GetIntDef(), TEXT("NextLineIndex")));	// Int Out
-		Sig.Outputs.Add(FNiagaraVariable(FNiagaraTypeDefinition::GetFloatDef(), TEXT("PrevWeight")));	// Float Out
+		Sig.Inputs.Add(FNiagaraVariable(FNiagaraTypeDefinition(GetClass()), TEXT("CSV")));					// CSV in
+		Sig.Inputs.Add(FNiagaraVariable(FNiagaraTypeDefinition::GetIntDef(), TEXT("N")));					// Point Number In
+		Sig.Inputs.Add(FNiagaraVariable(FNiagaraTypeDefinition::GetFloatDef(), TEXT("Time")));				// Time in
+		Sig.Outputs.Add(FNiagaraVariable(FNiagaraTypeDefinition::GetIntDef(), TEXT("PreviousRow")));		// Int Out
+		Sig.Outputs.Add(FNiagaraVariable(FNiagaraTypeDefinition::GetIntDef(), TEXT("NextRow")));			// Int Out
+		Sig.Outputs.Add(FNiagaraVariable(FNiagaraTypeDefinition::GetFloatDef(), TEXT("PrevWeight")));		// Float Out
 
 		OutFunctions.Add(Sig);
 	}
@@ -267,7 +279,7 @@ void UNiagaraDataInterfaceHoudiniCSV::GetFunctions(TArray<FNiagaraFunctionSignat
 		Sig.bMemberFunction = true;
 		Sig.bRequiresContext = false;
 		Sig.Inputs.Add(FNiagaraVariable(FNiagaraTypeDefinition(GetClass()), TEXT("CSV")));				// CSV in
-		Sig.Inputs.Add(FNiagaraVariable(FNiagaraTypeDefinition::GetFloatDef(), TEXT("N")));				// Point Number In
+		Sig.Inputs.Add(FNiagaraVariable(FNiagaraTypeDefinition::GetIntDef(), TEXT("N")));				// Point Number In
 		Sig.Inputs.Add(FNiagaraVariable(FNiagaraTypeDefinition::GetFloatDef(), TEXT("Time")));		    // Time in
 		Sig.Outputs.Add(FNiagaraVariable(FNiagaraTypeDefinition::GetVec3Def(), TEXT("Position")));		// Vector3 Out
 
@@ -281,9 +293,9 @@ void UNiagaraDataInterfaceHoudiniCSV::GetFunctions(TArray<FNiagaraFunctionSignat
 		Sig.bMemberFunction = true;
 		Sig.bRequiresContext = false;
 		Sig.Inputs.Add(FNiagaraVariable(FNiagaraTypeDefinition(GetClass()), TEXT("CSV")));				// CSV in
-		Sig.Inputs.Add(FNiagaraVariable(FNiagaraTypeDefinition::GetFloatDef(), TEXT("N")));				// Point Number In		
-		Sig.Inputs.Add(FNiagaraVariable(FNiagaraTypeDefinition::GetFloatDef(), TEXT("Time")));		    // Time in
-		Sig.Inputs.Add(FNiagaraVariable(FNiagaraTypeDefinition::GetFloatDef(), TEXT("Col")));			// Col Index In
+		Sig.Inputs.Add(FNiagaraVariable(FNiagaraTypeDefinition::GetIntDef(), TEXT("N")));				// Point Number In		
+		Sig.Inputs.Add(FNiagaraVariable(FNiagaraTypeDefinition::GetIntDef(), TEXT("Col")));				// Col Index In
+		Sig.Inputs.Add(FNiagaraVariable(FNiagaraTypeDefinition::GetFloatDef(), TEXT("Time")));		    // Time in		
 		Sig.Outputs.Add(FNiagaraVariable(FNiagaraTypeDefinition::GetFloatDef(), TEXT("Value")));		// Float Out
 
 		OutFunctions.Add(Sig);
@@ -296,9 +308,9 @@ void UNiagaraDataInterfaceHoudiniCSV::GetFunctions(TArray<FNiagaraFunctionSignat
 		Sig.bMemberFunction = true;
 		Sig.bRequiresContext = false;
 		Sig.Inputs.Add(FNiagaraVariable(FNiagaraTypeDefinition(GetClass()), TEXT("CSV")));				// CSV in
-		Sig.Inputs.Add(FNiagaraVariable(FNiagaraTypeDefinition::GetFloatDef(), TEXT("N")));				// Point Number In
-		Sig.Inputs.Add(FNiagaraVariable(FNiagaraTypeDefinition::GetFloatDef(), TEXT("Time")));		    // Time in
-		Sig.Inputs.Add(FNiagaraVariable(FNiagaraTypeDefinition::GetFloatDef(), TEXT("Col")));			// Col Index In
+		Sig.Inputs.Add(FNiagaraVariable(FNiagaraTypeDefinition::GetIntDef(), TEXT("N")));				// Point Number In
+		Sig.Inputs.Add(FNiagaraVariable(FNiagaraTypeDefinition::GetIntDef(), TEXT("Col")));				// Col Index In
+		Sig.Inputs.Add(FNiagaraVariable(FNiagaraTypeDefinition::GetFloatDef(), TEXT("Time")));		    // Time in		
 		Sig.Outputs.Add(FNiagaraVariable(FNiagaraTypeDefinition::GetVec3Def(), TEXT("Value")));			// Vector3 Out
 
 		OutFunctions.Add(Sig);
@@ -396,9 +408,9 @@ DEFINE_NDI_RAW_FUNC_BINDER(UNiagaraDataInterfaceHoudiniCSV, GetCSVPosition);
 DEFINE_NDI_RAW_FUNC_BINDER(UNiagaraDataInterfaceHoudiniCSV, GetCSVNormal);
 DEFINE_NDI_RAW_FUNC_BINDER(UNiagaraDataInterfaceHoudiniCSV, GetCSVTime);
 DEFINE_NDI_RAW_FUNC_BINDER(UNiagaraDataInterfaceHoudiniCSV, GetCSVPositionAndTime);
-DEFINE_NDI_RAW_FUNC_BINDER(UNiagaraDataInterfaceHoudiniCSV, GetLastParticleIndexAtTime);
+DEFINE_NDI_RAW_FUNC_BINDER(UNiagaraDataInterfaceHoudiniCSV, GetLastRowIndexAtTime);
 DEFINE_NDI_RAW_FUNC_BINDER(UNiagaraDataInterfaceHoudiniCSV, GetParticleIndexesToSpawnAtTime);
-DEFINE_NDI_RAW_FUNC_BINDER(UNiagaraDataInterfaceHoudiniCSV, GetParticleIndexesAtTime);
+DEFINE_NDI_RAW_FUNC_BINDER(UNiagaraDataInterfaceHoudiniCSV, GetRowIndexesForParticleAtTime);
 DEFINE_NDI_RAW_FUNC_BINDER(UNiagaraDataInterfaceHoudiniCSV, GetParticlePositionAtTime);
 DEFINE_NDI_RAW_FUNC_BINDER(UNiagaraDataInterfaceHoudiniCSV, GetParticleValueAtTime);
 DEFINE_NDI_RAW_FUNC_BINDER(UNiagaraDataInterfaceHoudiniCSV, GetParticleVectorValueAtTime);
@@ -406,7 +418,7 @@ void UNiagaraDataInterfaceHoudiniCSV::GetVMExternalFunction(const FVMExternalFun
 {
     if (BindingInfo.Name == TEXT("GetCSVFloatValue") && BindingInfo.GetNumInputs() == 2 && BindingInfo.GetNumOutputs() == 1)
     {
-		TNDIParamBinder<0, float, TNDIParamBinder<1, float, NDI_RAW_FUNC_BINDER(UNiagaraDataInterfaceHoudiniCSV, GetCSVFloatValue)>>::Bind(this, BindingInfo, InstanceData, OutFunc);
+		TNDIParamBinder<0, int32, TNDIParamBinder<1, int32, NDI_RAW_FUNC_BINDER(UNiagaraDataInterfaceHoudiniCSV, GetCSVFloatValue)>>::Bind(this, BindingInfo, InstanceData, OutFunc);
     }
     /*else if (BindingInfo.Name == TEXT("GetCSVFloatValueByString") && BindingInfo.GetNumInputs() == 2 && BindingInfo.GetNumOutputs() == 1)
     {
@@ -414,53 +426,55 @@ void UNiagaraDataInterfaceHoudiniCSV::GetVMExternalFunction(const FVMExternalFun
     }*/
     else if (BindingInfo.Name == TEXT("GetCSVVectorValue") && BindingInfo.GetNumInputs() == 2 && BindingInfo.GetNumOutputs() == 3)
     {
-		TNDIParamBinder<0, float, TNDIParamBinder<1, float, NDI_RAW_FUNC_BINDER(UNiagaraDataInterfaceHoudiniCSV, GetCSVVectorValue)>>::Bind(this, BindingInfo, InstanceData, OutFunc);
+		TNDIParamBinder<0, int32, TNDIParamBinder<1, int32, NDI_RAW_FUNC_BINDER(UNiagaraDataInterfaceHoudiniCSV, GetCSVVectorValue)>>::Bind(this, BindingInfo, InstanceData, OutFunc);
     }    
     else if (BindingInfo.Name == TEXT("GetCSVPosition") && BindingInfo.GetNumInputs() == 1 && BindingInfo.GetNumOutputs() == 3)
     {
-		TNDIParamBinder<0, float, NDI_RAW_FUNC_BINDER(UNiagaraDataInterfaceHoudiniCSV, GetCSVPosition)>::Bind(this, BindingInfo, InstanceData, OutFunc);
+		TNDIParamBinder<0, int32, NDI_RAW_FUNC_BINDER(UNiagaraDataInterfaceHoudiniCSV, GetCSVPosition)>::Bind(this, BindingInfo, InstanceData, OutFunc);
     }
     else if (BindingInfo.Name == TEXT("GetCSVNormal") && BindingInfo.GetNumInputs() == 1 && BindingInfo.GetNumOutputs() == 3)
     {
-		TNDIParamBinder<0, float, NDI_RAW_FUNC_BINDER(UNiagaraDataInterfaceHoudiniCSV, GetCSVNormal)>::Bind(this, BindingInfo, InstanceData, OutFunc);
+		TNDIParamBinder<0, int32, NDI_RAW_FUNC_BINDER(UNiagaraDataInterfaceHoudiniCSV, GetCSVNormal)>::Bind(this, BindingInfo, InstanceData, OutFunc);
     }
     else if (BindingInfo.Name == TEXT("GetCSVTime") && BindingInfo.GetNumInputs() == 1 && BindingInfo.GetNumOutputs() == 1)
     {
-		TNDIParamBinder<0, float, NDI_RAW_FUNC_BINDER(UNiagaraDataInterfaceHoudiniCSV, GetCSVTime)>::Bind(this, BindingInfo, InstanceData, OutFunc);
+		TNDIParamBinder<0, int32, NDI_RAW_FUNC_BINDER(UNiagaraDataInterfaceHoudiniCSV, GetCSVTime)>::Bind(this, BindingInfo, InstanceData, OutFunc);
     }
     else if (BindingInfo.Name == TEXT("GetCSVPositionAndTime") && BindingInfo.GetNumInputs() == 1 && BindingInfo.GetNumOutputs() == 4)
     {
-		TNDIParamBinder<0, float, NDI_RAW_FUNC_BINDER(UNiagaraDataInterfaceHoudiniCSV, GetCSVPositionAndTime)>::Bind(this, BindingInfo, InstanceData, OutFunc);
+		TNDIParamBinder<0, int32, NDI_RAW_FUNC_BINDER(UNiagaraDataInterfaceHoudiniCSV, GetCSVPositionAndTime)>::Bind(this, BindingInfo, InstanceData, OutFunc);
     }
-    else if ( BindingInfo.Name == TEXT("GetNumberOfPointsInCSV") && BindingInfo.GetNumInputs() == 0 && BindingInfo.GetNumOutputs() == 1 )
+    else if ( BindingInfo.Name == TEXT("GetNumberOfParticlesInCSV") && BindingInfo.GetNumInputs() == 0 && BindingInfo.GetNumOutputs() == 1 )
     {
-		OutFunc = FVMExternalFunction::CreateUObject(this, &UNiagaraDataInterfaceHoudiniCSV::GetNumberOfPointsInCSV);
+		OutFunc = FVMExternalFunction::CreateUObject(this, &UNiagaraDataInterfaceHoudiniCSV::GetNumberOfParticlesInCSV);
     }
-    else if (BindingInfo.Name == TEXT("GetLastParticleIndexAtTime") && BindingInfo.GetNumInputs() == 1 && BindingInfo.GetNumOutputs() == 1)
+	else if (BindingInfo.Name == TEXT("GetNumberOfRowsInCSV") && BindingInfo.GetNumInputs() == 0 && BindingInfo.GetNumOutputs() == 1)
+	{
+		OutFunc = FVMExternalFunction::CreateUObject(this, &UNiagaraDataInterfaceHoudiniCSV::GetNumberOfRowsInCSV);
+	}
+    else if (BindingInfo.Name == TEXT("GetLastRowIndexAtTime") && BindingInfo.GetNumInputs() == 1 && BindingInfo.GetNumOutputs() == 1)
     {
-		TNDIParamBinder<0, float, NDI_RAW_FUNC_BINDER(UNiagaraDataInterfaceHoudiniCSV, GetLastParticleIndexAtTime)>::Bind(this, BindingInfo, InstanceData, OutFunc);
+		TNDIParamBinder<0, float, NDI_RAW_FUNC_BINDER(UNiagaraDataInterfaceHoudiniCSV, GetLastRowIndexAtTime)>::Bind(this, BindingInfo, InstanceData, OutFunc);
     }
     else if (BindingInfo.Name == TEXT("GetParticleIndexesToSpawnAtTime") && BindingInfo.GetNumInputs() == 1 && BindingInfo.GetNumOutputs() == 3)
     {
 		TNDIParamBinder<0, float, NDI_RAW_FUNC_BINDER(UNiagaraDataInterfaceHoudiniCSV, GetParticleIndexesToSpawnAtTime)>::Bind(this, BindingInfo, InstanceData, OutFunc);
     }
-	else if (BindingInfo.Name == TEXT("GetParticleIndexesAtTime") && BindingInfo.GetNumInputs() == 2 && BindingInfo.GetNumOutputs() == 3)
+	else if (BindingInfo.Name == TEXT("GetRowIndexesForParticleAtTime") && BindingInfo.GetNumInputs() == 2 && BindingInfo.GetNumOutputs() == 3)
 	{
-		TNDIParamBinder<0, float, TNDIParamBinder<1, float, NDI_RAW_FUNC_BINDER(UNiagaraDataInterfaceHoudiniCSV, GetParticleIndexesAtTime)>>::Bind(this, BindingInfo, InstanceData, OutFunc);
+		TNDIParamBinder<0, int32, TNDIParamBinder<1, float, NDI_RAW_FUNC_BINDER(UNiagaraDataInterfaceHoudiniCSV, GetRowIndexesForParticleAtTime)>>::Bind(this, BindingInfo, InstanceData, OutFunc);
 	}
 	else if (BindingInfo.Name == TEXT("GetParticlePositionAtTime") && BindingInfo.GetNumInputs() == 2 && BindingInfo.GetNumOutputs() == 3)
 	{
-		TNDIParamBinder<0, float, TNDIParamBinder<1, float, NDI_RAW_FUNC_BINDER(UNiagaraDataInterfaceHoudiniCSV, GetParticlePositionAtTime)>>::Bind(this, BindingInfo, InstanceData, OutFunc);
+		TNDIParamBinder<0, int32, TNDIParamBinder<1, float, NDI_RAW_FUNC_BINDER(UNiagaraDataInterfaceHoudiniCSV, GetParticlePositionAtTime)>>::Bind(this, BindingInfo, InstanceData, OutFunc);
 	}
-
 	else if (BindingInfo.Name == TEXT("GetParticleValueAtTime") && BindingInfo.GetNumInputs() == 3 && BindingInfo.GetNumOutputs() == 1)
 	{
-		TNDIParamBinder<0, float, TNDIParamBinder<1, float, TNDIParamBinder<2, float, NDI_RAW_FUNC_BINDER(UNiagaraDataInterfaceHoudiniCSV, GetParticleValueAtTime)>>>::Bind(this, BindingInfo, InstanceData, OutFunc);
+		TNDIParamBinder<0, int32, TNDIParamBinder<1, int32, TNDIParamBinder<2, float, NDI_RAW_FUNC_BINDER(UNiagaraDataInterfaceHoudiniCSV, GetParticleValueAtTime)>>>::Bind(this, BindingInfo, InstanceData, OutFunc);
 	}
-
 	else if (BindingInfo.Name == TEXT("GetParticleVectorValueAtTime") && BindingInfo.GetNumInputs() == 2 && BindingInfo.GetNumOutputs() == 3)
 	{
-		TNDIParamBinder<0, float, TNDIParamBinder<1, float, TNDIParamBinder<2, float, NDI_RAW_FUNC_BINDER(UNiagaraDataInterfaceHoudiniCSV, GetParticleVectorValueAtTime)>>>::Bind(this, BindingInfo, InstanceData, OutFunc);
+		TNDIParamBinder<0, int32, TNDIParamBinder<1, int32, TNDIParamBinder<2, float, NDI_RAW_FUNC_BINDER(UNiagaraDataInterfaceHoudiniCSV, GetParticleVectorValueAtTime)>>>::Bind(this, BindingInfo, InstanceData, OutFunc);
 	}
     else
     {
@@ -481,8 +495,8 @@ void UNiagaraDataInterfaceHoudiniCSV::GetCSVFloatValue(FVectorVMContext& Context
 
     for ( int32 i = 0; i < Context.NumInstances; ++i )
     {
-		float row = RowParam.Get();
-		float col = ColParam.Get();
+		int32 row = RowParam.Get();
+		int32 col = ColParam.Get();
 	
 		float value = 0.0f;
 		if ( CSVFile )
@@ -501,28 +515,25 @@ void UNiagaraDataInterfaceHoudiniCSV::GetCSVVectorValue( FVectorVMContext& Conte
     RowParamType RowParam(Context);
     ColParamType ColParam(Context);
 
-    //FRegisterHandler<FVector> OutValue(Context);
 	FRegisterHandler<float> OutVectorX(Context);
 	FRegisterHandler<float> OutVectorY(Context);
 	FRegisterHandler<float> OutVectorZ(Context);
 
     for (int32 i = 0; i < Context.NumInstances; ++i)
     {
-		float row = RowParam.Get();
-		float col = ColParam.Get();
+		int32 row = RowParam.Get();
+		int32 col = ColParam.Get();
 
 		FVector V = FVector::ZeroVector;
 		if ( CSVFile )
 			CSVFile->GetCSVVectorValue(row, col, V);
 
-		//*OutValue.GetDest() = value;
 		*OutVectorX.GetDest() = V.X;
 		*OutVectorY.GetDest() = V.Y;
 		*OutVectorZ.GetDest() = V.Z;
 
 		RowParam.Advance();
 		ColParam.Advance();
-		//OutValue.Advance();
 		OutVectorX.Advance();
 		OutVectorY.Advance();
 		OutVectorZ.Advance();
@@ -539,7 +550,7 @@ void UNiagaraDataInterfaceHoudiniCSV::GetCSVFloatValueByString(FVectorVMContext&
 
     for ( int32 i = 0; i < Context.NumInstances; ++i )
     {
-		float row = RowParam.Get();
+		int32 row = RowParam.Get();
 		FString colTitle = ColTitleParam.Get();
 	
 		float value = 0.0f;
@@ -554,81 +565,81 @@ void UNiagaraDataInterfaceHoudiniCSV::GetCSVFloatValueByString(FVectorVMContext&
 }
 */
 
-template<typename NParamType>
+template<typename RowParamType>
 void UNiagaraDataInterfaceHoudiniCSV::GetCSVPosition(FVectorVMContext& Context)
 {
-    NParamType NParam(Context);
+	RowParamType RowParam(Context);
     FRegisterHandler<float> OutSampleX(Context);
     FRegisterHandler<float> OutSampleY(Context);
     FRegisterHandler<float> OutSampleZ(Context);
 
     for (int32 i = 0; i < Context.NumInstances; ++i)
     {
-		float N = NParam.Get();
+		int32 row = RowParam.Get();
 
 		FVector V = FVector::ZeroVector;
-		if (CSVFile)
-			CSVFile->GetCSVPositionValue(N, V);
+		if ( CSVFile )
+			CSVFile->GetCSVPositionValue( row, V );
 
 		*OutSampleX.GetDest() = V.X;
 		*OutSampleY.GetDest() = V.Y;
 		*OutSampleZ.GetDest() = V.Z;
-		NParam.Advance();
+		RowParam.Advance();
 		OutSampleX.Advance();
 		OutSampleY.Advance();
 		OutSampleZ.Advance();
     }
 }
 
-template<typename NParamType>
+template<typename RowParamType>
 void UNiagaraDataInterfaceHoudiniCSV::GetCSVNormal(FVectorVMContext& Context)
 {
-    NParamType NParam(Context);
+	RowParamType RowParam(Context);
     FRegisterHandler<float> OutSampleX(Context);
     FRegisterHandler<float> OutSampleY(Context);
     FRegisterHandler<float> OutSampleZ(Context);
 
     for (int32 i = 0; i < Context.NumInstances; ++i)
     {
-		float N = NParam.Get();
+		int32 row = RowParam.Get();
 
 		FVector V = FVector::ZeroVector;
-		if (CSVFile)
-			CSVFile->GetCSVNormalValue(N, V);
+		if ( CSVFile )
+			CSVFile->GetCSVNormalValue( row, V );
 
 		*OutSampleX.GetDest() = V.X;
 		*OutSampleY.GetDest() = V.Y;
 		*OutSampleZ.GetDest() = V.Z;
-		NParam.Advance();
+		RowParam.Advance();
 		OutSampleX.Advance();
 		OutSampleY.Advance();
 		OutSampleZ.Advance();
     }
 }
 
-template<typename NParamType>
+template<typename RowParamType>
 void UNiagaraDataInterfaceHoudiniCSV::GetCSVTime(FVectorVMContext& Context)
 {
-    NParamType NParam(Context);
+	RowParamType RowParam(Context);
     FRegisterHandler<float> OutValue(Context);
 
     for (int32 i = 0; i < Context.NumInstances; ++i)
     {
-		float N = NParam.Get();
+		int32 row = RowParam.Get();
 
 		float value = 0.0f;
-		if (CSVFile)
-			CSVFile->GetCSVTimeValue(N, value);
+		if ( CSVFile )
+			CSVFile->GetCSVTimeValue( row, value );
 
 		*OutValue.GetDest() = value;
-		NParam.Advance();
+		RowParam.Advance();
 		OutValue.Advance();
     }
 }
 
 // Returns the last index of the particles that should be spawned at time t
 template<typename TimeParamType>
-void UNiagaraDataInterfaceHoudiniCSV::GetLastParticleIndexAtTime(FVectorVMContext& Context)
+void UNiagaraDataInterfaceHoudiniCSV::GetLastRowIndexAtTime(FVectorVMContext& Context)
 {
     TimeParamType TimeParam(Context);
     FRegisterHandler<int32> OutValue(Context);
@@ -686,7 +697,7 @@ void UNiagaraDataInterfaceHoudiniCSV::GetParticleIndexesToSpawnAtTime( FVectorVM
 				else
 				{
 					// The last time value in the CSV is lower than t, spawn everything if we didnt already!
-					if ( value >= CSVFile->GetNumberOfPointsInCSV() )
+					if ( value >= CSVFile->GetNumberOfParticlesInCSV() )
 						value = value - 1;
 
 					if ( value == LastSpawnIndex )
@@ -720,10 +731,10 @@ void UNiagaraDataInterfaceHoudiniCSV::GetParticleIndexesToSpawnAtTime( FVectorVM
     }
 }
 
-template<typename NParamType>
+template<typename RowParamType>
 void UNiagaraDataInterfaceHoudiniCSV::GetCSVPositionAndTime(FVectorVMContext& Context)
 {
-    NParamType NParam(Context);
+    RowParamType RowParam(Context);
 
     FRegisterHandler<float> OutPosX(Context);
     FRegisterHandler<float> OutPosY(Context);
@@ -732,14 +743,14 @@ void UNiagaraDataInterfaceHoudiniCSV::GetCSVPositionAndTime(FVectorVMContext& Co
 
     for (int32 i = 0; i < Context.NumInstances; ++i)
     {
-		float N = NParam.Get();
+		int32 row = RowParam.Get();
 
 		float timeValue = 0.0f;
 		FVector posVector = FVector::ZeroVector;
 		if ( CSVFile )
 		{
-			CSVFile->GetCSVTimeValue(N, timeValue);
-			CSVFile->GetCSVPositionValue(N, posVector);
+			CSVFile->GetCSVTimeValue( row, timeValue);
+			CSVFile->GetCSVPositionValue( row, posVector);
 		}
 
 		*OutPosX.GetDest() = posVector.X;
@@ -748,7 +759,7 @@ void UNiagaraDataInterfaceHoudiniCSV::GetCSVPositionAndTime(FVectorVMContext& Co
 
 		*OutTime.GetDest() = timeValue;
 
-		NParam.Advance();
+		RowParam.Advance();
 		OutPosX.Advance();
 		OutPosY.Advance();
 		OutPosZ.Advance();
@@ -757,7 +768,7 @@ void UNiagaraDataInterfaceHoudiniCSV::GetCSVPositionAndTime(FVectorVMContext& Co
 }
 
 template<typename NParamType, typename TimeParamType>
-void UNiagaraDataInterfaceHoudiniCSV::GetParticleIndexesAtTime(FVectorVMContext& Context)
+void UNiagaraDataInterfaceHoudiniCSV::GetRowIndexesForParticleAtTime(FVectorVMContext& Context)
 {
 	NParamType NParam(Context);
 	TimeParamType TimeParam(Context);
@@ -768,7 +779,7 @@ void UNiagaraDataInterfaceHoudiniCSV::GetParticleIndexesAtTime(FVectorVMContext&
 
 	for (int32 i = 0; i < Context.NumInstances; ++i)
     {
-		float N = NParam.Get();
+		int32 N = NParam.Get();
 		float time = TimeParam.Get();
 
 		float weight = 0.0f;
@@ -797,14 +808,13 @@ void UNiagaraDataInterfaceHoudiniCSV::GetParticlePositionAtTime( FVectorVMContex
 	NParamType NParam(Context);
 	TimeParamType TimeParam(Context);
 
-	//FRegisterHandler<FVector> OutValue(Context);
 	FRegisterHandler<float> OutPosX(Context);
 	FRegisterHandler<float> OutPosY(Context);
 	FRegisterHandler<float> OutPosZ(Context);
 
 	for (int32 i = 0; i < Context.NumInstances; ++i)
     {
-		float N = NParam.Get();
+		int32 N = NParam.Get();
 		float time = TimeParam.Get();
 
 		FVector posVector = FVector::ZeroVector;
@@ -813,14 +823,12 @@ void UNiagaraDataInterfaceHoudiniCSV::GetParticlePositionAtTime( FVectorVMContex
 			CSVFile->GetParticlePositionAtTime(N, time, posVector);
 		}		
 
-		//*OutValue.GetDest() = posVector;
 		*OutPosX.GetDest() = posVector.X;
 		*OutPosY.GetDest() = posVector.Y;
 		*OutPosZ.GetDest() = posVector.Z;
 
 		NParam.Advance();
 		TimeParam.Advance();
-		//OutValue.Advance();
 		OutPosX.Advance();
 		OutPosY.Advance();
 		OutPosZ.Advance();
@@ -838,21 +846,22 @@ void UNiagaraDataInterfaceHoudiniCSV::GetParticleValueAtTime(FVectorVMContext& C
 
 	for (int32 i = 0; i < Context.NumInstances; ++i)
 	{
-		float N = NParam.Get();
-		float time = TimeParam.Get();
-		float Col = ColParam.Get();
+		int32 N = NParam.Get();
+		int32 Col = ColParam.Get();
+		float time = TimeParam.Get();		
 
 		float Value = 0.0f;
 		if ( CSVFile )
 		{
-			CSVFile->GetParticleValueAtTime(N, Col, time, Value);
+			CSVFile->GetParticleValueAtTime( N, Col, time, Value );
 		}
 
 		*OutValue.GetDest() = Value;
 
 		NParam.Advance();
-		TimeParam.Advance();
 		ColParam.Advance();
+		TimeParam.Advance();
+
 		OutValue.Advance();
 	}
 }
@@ -861,8 +870,8 @@ template<typename NParamType, typename ColParamType, typename TimeParamType>
 void UNiagaraDataInterfaceHoudiniCSV::GetParticleVectorValueAtTime(FVectorVMContext& Context)
 {
 	NParamType NParam(Context);
-	TimeParamType TimeParam(Context);
 	ColParamType ColParam(Context);
+	TimeParamType TimeParam(Context);	
 
 	FRegisterHandler<float> OutPosX(Context);
 	FRegisterHandler<float> OutPosY(Context);
@@ -870,9 +879,9 @@ void UNiagaraDataInterfaceHoudiniCSV::GetParticleVectorValueAtTime(FVectorVMCont
 
 	for (int32 i = 0; i < Context.NumInstances; ++i)
 	{
-		float N = NParam.Get();
-		float time = TimeParam.Get();
-		float Col = ColParam.Get();
+		int32 N = NParam.Get();
+		int32 Col = ColParam.Get();
+		float time = TimeParam.Get();		
 
 		FVector posVector = FVector::ZeroVector;
 		if ( CSVFile )
@@ -885,19 +894,27 @@ void UNiagaraDataInterfaceHoudiniCSV::GetParticleVectorValueAtTime(FVectorVMCont
 		*OutPosZ.GetDest() = posVector.Z;
 
 		NParam.Advance();
-		TimeParam.Advance();
 		ColParam.Advance();
+		TimeParam.Advance();
+		
 		OutPosX.Advance();
 		OutPosY.Advance();
 		OutPosZ.Advance();
 	}
 }
 
-void UNiagaraDataInterfaceHoudiniCSV::GetNumberOfPointsInCSV(FVectorVMContext& Context)
+void UNiagaraDataInterfaceHoudiniCSV::GetNumberOfRowsInCSV(FVectorVMContext& Context)
+{
+    FRegisterHandler<int32> OutNumRows(Context);
+    *OutNumRows.GetDest() = CSVFile ? CSVFile->GetNumberOfLinesInCSV() : 0;
+	OutNumRows.Advance();
+}
+
+void UNiagaraDataInterfaceHoudiniCSV::GetNumberOfParticlesInCSV(FVectorVMContext& Context)
 {
-    FRegisterHandler<int32> OutNumPoints(Context);
-    *OutNumPoints.GetDest() = CSVFile ? CSVFile->NumberOfLines : 0;
-    OutNumPoints.Advance();
+	FRegisterHandler<int32> OutNumParticles(Context);
+	*OutNumParticles.GetDest() = CSVFile ? CSVFile->GetNumberOfParticlesInCSV() : 0;
+	OutNumParticles.Advance();
 }
 
 #undef LOCTEXT_NAMESPACE
diff --git a/Engine/Plugins/FX/HoudiniNiagara/Source/HoudiniNiagara/Public/HoudiniCSV.h b/Engine/Plugins/FX/HoudiniNiagara/Source/HoudiniNiagara/Public/HoudiniCSV.h
--- a/Engine/Plugins/FX/HoudiniNiagara/Source/HoudiniNiagara/Public/HoudiniCSV.h
+++ b/Engine/Plugins/FX/HoudiniNiagara/Source/HoudiniNiagara/Public/HoudiniCSV.h
@@ -46,7 +46,7 @@ class HOUDININIAGARA_API UHoudiniCSV : public UObject
 	void SetFileName( const FString& TheFilename );
 
 	// Returns the number of points found in the CSV file
-	int32 GetNumberOfPointsInCSV();
+	int32 GetNumberOfParticlesInCSV();
 
 	// Returns the number of points found in the CSV file
 	int32 GetNumberOfLinesInCSV();
