Niagara - Houdini:

Houdini Niagara Data Interface:
Fixed incorrect behavior of the GetLastRowIndexAtTime and GetParticleIndexesToSpawnAtTime functions due to supporting particle update over time.

Houdini CSV Asset:
Fixed missing UPropery for SpawnTimes and LifeValues array.

diff --git a/Engine/Plugins/FX/HoudiniNiagara/Source/HoudiniNiagara/Classes/NiagaraDataInterfaceHoudiniCSV.h b/Engine/Plugins/FX/HoudiniNiagara/Source/HoudiniNiagara/Classes/NiagaraDataInterfaceHoudiniCSV.h
--- a/Engine/Plugins/FX/HoudiniNiagara/Source/HoudiniNiagara/Classes/NiagaraDataInterfaceHoudiniCSV.h
+++ b/Engine/Plugins/FX/HoudiniNiagara/Source/HoudiniNiagara/Classes/NiagaraDataInterfaceHoudiniCSV.h
@@ -137,6 +137,10 @@ protected:
 	bool GPUBufferDirty;
 
 	// Last Index used to spawn particles
+	//UPROPERTY()
+	//int32 LastSpawnIndex;
+
+	// Last Spawned Particle ID
 	UPROPERTY()
-	int32 LastSpawnIndex;
+	int32 LastSpawnedParticleID;
 };
diff --git a/Engine/Plugins/FX/HoudiniNiagara/Source/HoudiniNiagara/Private/HoudiniCSV.cpp b/Engine/Plugins/FX/HoudiniNiagara/Source/HoudiniNiagara/Private/HoudiniCSV.cpp
--- a/Engine/Plugins/FX/HoudiniNiagara/Source/HoudiniNiagara/Private/HoudiniCSV.cpp
+++ b/Engine/Plugins/FX/HoudiniNiagara/Source/HoudiniNiagara/Private/HoudiniCSV.cpp
@@ -517,85 +517,111 @@ int32 UHoudiniCSV::GetNumberOfLinesInCSV()
 	return NumberOfLines;
 }
 
-// Get the last index of the particles with a time value smaller or equal to desiredTime
-// If the CSV file doesn't have time informations, returns false and set the LastIndex to the last particle
-// If desiredTime is smaller than the first particle, LastIndex will be set to -1
-// If desiredTime is higher than the last particle in the csv file, LastIndex will be set to the last particle's index
-bool UHoudiniCSV::GetLastParticleIndexAtTime( const float& desiredTime, int32& lastIndex )
+// Get the last row index for a given time value (the row with a time smaller or equal to desiredTime)
+// If the CSV file doesn't have time informations, returns false and set the LastRowIndex to the last line in the file
+// If desiredTime is smaller than the time value in the first row, LastRowIndex will be set to -1
+// If desiredTime is higher than the last time value in the last row of the csv file, LastIndex will be set to the last row's index
+bool UHoudiniCSV::GetLastRowIndexAtTime(const float& desiredTime, int32& lastRowIndex)
 {
-    // If we dont have proper time info, always return the last line
-    if ( TimeColumnIndex < 0 || TimeColumnIndex >= NumberOfColumns )
-    {
-		lastIndex = NumberOfLines - 1;
+	// If we dont have proper time info, always return the last line
+	if (TimeColumnIndex < 0 || TimeColumnIndex >= NumberOfColumns)
+	{
+		lastRowIndex = NumberOfLines - 1;
 		return false;
-    }
+	}
 
-	/*
-    float temp_time = 0.0f;
-	if ( SpawnTimes.IsValidIndex(NumberOfParticles - 1) )
+	float temp_time = 0.0f;
+	// Check first if we have anything to spawn at the current time by looking at the last value
+	if ( GetCSVTimeValue( NumberOfLines - 1, temp_time ) && temp_time < desiredTime )
 	{
-		if ( SpawnTimes[ NumberOfParticles - 1 ] < desiredTime )
-		{
-			// We didn't find a suitable index because the desired time is higher than our last time value
-			lastIndex = NumberOfLines - 1;
-			return true;
-		}
+		// We didn't find a suitable index because the desired time is higher than our last time value
+		lastRowIndex = NumberOfLines - 1;
+		return true;
 	}
 
-	// Iterates through all the particles
-	lastIndex = INDEX_NONE;
-	for (int32 n = 0; n < NumberOfParticles; n++)
+	// Iterates through all the lines
+	lastRowIndex = INDEX_NONE;
+	for ( int32 n = 0; n < NumberOfLines; n++ )
 	{
-		temp_time = SpawnTimes[ NumberOfParticles - 1 ];
-
-
-
-		if (GetCSVTimeValue(n, temp_time))
+		if ( GetCSVTimeValue( n, temp_time ) )
 		{
-			if (temp_time == desiredTime)
+			if ( temp_time == desiredTime )
 			{
-				lastIndex = n;
+				lastRowIndex = n;
 			}
-			else if (temp_time > desiredTime)
+			else if ( temp_time > desiredTime )
 			{
-				lastIndex = n - 1;
+				lastRowIndex = n - 1;
 				return true;
 			}
 		}
 	}
-	*/
-	float temp_time = 0.0f;
-    // Check first if we have anything to spawn at the current time by looking at the last value
-    if ( GetCSVTimeValue( NumberOfLines - 1, temp_time ) && temp_time < desiredTime )
+
+	// We didn't find a suitable index because the desired time is higher than our last time value
+	if ( lastRowIndex == INDEX_NONE )
+		lastRowIndex = NumberOfLines - 1;
+
+	return true;
+}
+// Get the last index of the particles with a time value smaller or equal to desiredTime
+// If the CSV file doesn't have time informations, returns false and set the LastIndex to the last particle
+// If desiredTime is smaller than the first particle, LastIndex will be set to -1
+// If desiredTime is higher than the last particle in the csv file, LastIndex will be set to the last particle's index
+bool UHoudiniCSV::GetLastParticleIndexAtTime( const float& desiredTime, int32& lastID )
+{
+    // If we dont have proper time info, always return the last particle
+    if ( TimeColumnIndex < 0 || TimeColumnIndex >= NumberOfColumns )
     {
+		lastID = NumberOfParticles - 1;
+		return false;
+    }
+
+    float temp_time = 0.0f;
+	if ( !SpawnTimes.IsValidIndex( NumberOfParticles - 1 ) )
+	{
+		lastID = NumberOfParticles - 1;
+		return false;
+	}
+
+	if ( SpawnTimes[ NumberOfParticles - 1 ] < desiredTime )
+	{
 		// We didn't find a suitable index because the desired time is higher than our last time value
-		lastIndex = NumberOfLines - 1;
+		lastID = NumberOfParticles - 1;
 		return true;
-    }	
+	}
 
-    // Iterates through all the particles
-    lastIndex = INDEX_NONE;
-    for ( int32 n = 0; n < NumberOfLines; n++ )
-    {
-		if (GetCSVTimeValue(n, temp_time))
+	// Iterates through all the particles
+	lastID = INDEX_NONE;
+	for ( int32 n = 0; n < NumberOfParticles; n++ )
+	{
+		temp_time = SpawnTimes[ n ];
+
+		if ( temp_time == desiredTime )
 		{
-			if (temp_time == desiredTime)
-			{
-				lastIndex = n;
-			}
-			else if (temp_time > desiredTime)
-			{
-				lastIndex = n - 1;
-				return true;
-			}
+			lastID = n;
 		}
-    }
+		else if ( temp_time > desiredTime )
+		{
+			lastID = n - 1;
+			return true;
+		}
+	}
+
+	return true;
+}
+
+bool UHoudiniCSV::GetParticleLifeAtTime( const int32& ParticleID, const float& DesiredTime, float& Value )
+{
+	if ( !SpawnTimes.IsValidIndex( ParticleID )  || !LifeValues.IsValidIndex( ParticleID ) )
+	{
+		Value = -1.0f;
+		return false;
+	}
 
-    // We didn't find a suitable index because the desired time is higher than our last time value
-    if ( lastIndex == INDEX_NONE )
-		lastIndex = NumberOfLines - 1;
+	if ( DesiredTime < SpawnTimes[ ParticleID ] )
+		return LifeValues[ ParticleID ];
 
-    return true;
+	return LifeValues[ ParticleID ] - ( DesiredTime - SpawnTimes[ ParticleID ] );
 }
 
 bool UHoudiniCSV::GetParticleValueAtTime( const int32& ParticleID, const int32& ColumnIndex, const float& desiredTime, float& Value )
diff --git a/Engine/Plugins/FX/HoudiniNiagara/Source/HoudiniNiagara/Private/NiagaraDataInterfaceHoudiniCSV.cpp b/Engine/Plugins/FX/HoudiniNiagara/Source/HoudiniNiagara/Private/NiagaraDataInterfaceHoudiniCSV.cpp
--- a/Engine/Plugins/FX/HoudiniNiagara/Source/HoudiniNiagara/Private/NiagaraDataInterfaceHoudiniCSV.cpp
+++ b/Engine/Plugins/FX/HoudiniNiagara/Source/HoudiniNiagara/Private/NiagaraDataInterfaceHoudiniCSV.cpp
@@ -32,7 +32,8 @@ UNiagaraDataInterfaceHoudiniCSV::UNiagaraDataInterfaceHoudiniCSV(FObjectInitiali
 	: Super(ObjectInitializer)
 {
     CSVFile = nullptr;
-    LastSpawnIndex = -1;
+    //LastSpawnIndex = -1;
+	LastSpawnedParticleID = -1;
 }
 
 void UNiagaraDataInterfaceHoudiniCSV::PostInitProperties()
@@ -45,14 +46,16 @@ void UNiagaraDataInterfaceHoudiniCSV::PostInitProperties()
     }
 
     GPUBufferDirty = true;
-    LastSpawnIndex = -1;
+    //LastSpawnIndex = -1;
+	LastSpawnedParticleID = -1;
 }
 
 void UNiagaraDataInterfaceHoudiniCSV::PostLoad()
 {
     Super::PostLoad();
     GPUBufferDirty = true;
-    LastSpawnIndex = -1;
+    //LastSpawnIndex = -1;
+	LastSpawnedParticleID = -1;
 }
 
 #if WITH_EDITOR
@@ -67,7 +70,8 @@ void UNiagaraDataInterfaceHoudiniCSV::PostEditChangeProperty(struct FPropertyCha
 		if (CSVFile)
 		{
 			GPUBufferDirty = true;
-			LastSpawnIndex = -1;
+			//LastSpawnIndex = -1;
+			LastSpawnedParticleID = -1;
 		}
     }
 }
@@ -650,7 +654,7 @@ void UNiagaraDataInterfaceHoudiniCSV::GetLastRowIndexAtTime(FVectorVMContext& Co
 
 		int32 value = 0;
 		if ( CSVFile )
-			CSVFile->GetLastParticleIndexAtTime( t, value );
+			CSVFile->GetLastRowIndexAtTime( t, value );
 
 		*OutValue.GetDest() = value;
 		TimeParam.Advance();
@@ -686,13 +690,13 @@ void UNiagaraDataInterfaceHoudiniCSV::GetParticleIndexesToSpawnAtTime( FVectorVM
 			{
 				// The CSV file has time informations
 				// First, detect if we need to reset LastSpawnIndex (after a loop)
-				if ( value < LastSpawnIndex )
-					LastSpawnIndex = -1;
+				if ( value < LastSpawnedParticleID )
+					LastSpawnedParticleID = -1;
 
 				if ( value < 0 )
 				{
 					// Nothing to spawn, t is lower than the particle's time
-					LastSpawnIndex = -1;
+					LastSpawnedParticleID = -1;
 				}
 				else
 				{
@@ -700,21 +704,21 @@ void UNiagaraDataInterfaceHoudiniCSV::GetParticleIndexesToSpawnAtTime( FVectorVM
 					if ( value >= CSVFile->GetNumberOfParticlesInCSV() )
 						value = value - 1;
 
-					if ( value == LastSpawnIndex )
+					if ( value == LastSpawnedParticleID)
 					{
-						// We've already spawned all the particles for this time
+						// We dont have any new particle to spawn
 						min = value;
 						max = value;
 						count = 0;
 					}
 					else
 					{
-						// Found particles to spawn at time t
-						min = LastSpawnIndex + 1;
+						// We have particles to spawn at time t
+						min = LastSpawnedParticleID + 1;
 						max = value;
 						count = max - min + 1;
 
-						LastSpawnIndex = max;
+						LastSpawnedParticleID = max;
 					}
 				}
 			}
diff --git a/Engine/Plugins/FX/HoudiniNiagara/Source/HoudiniNiagara/Public/HoudiniCSV.h b/Engine/Plugins/FX/HoudiniNiagara/Source/HoudiniNiagara/Public/HoudiniCSV.h
--- a/Engine/Plugins/FX/HoudiniNiagara/Source/HoudiniNiagara/Public/HoudiniCSV.h
+++ b/Engine/Plugins/FX/HoudiniNiagara/Source/HoudiniNiagara/Public/HoudiniCSV.h
@@ -72,11 +72,18 @@ class HOUDININIAGARA_API UHoudiniCSV : public UObject
 	// Returns a Normal Vector3 for a given point in the CSV file (converted to unreal's coordinate system)
 	bool GetCSVNormalValue( const int32& lineIndex, FVector& value );
 
+	// Get the last row index for a given time value (the row with a time smaller or equal to desiredTime)
+	// If the CSV file doesn't have time informations, returns false and set the LastRowIndex to the last line in the file
+	// If desiredTime is smaller than the time value in the first row, LastRowIndex will be set to -1
+	// If desiredTime is higher than the last time value in the last row of the csv file, LastIndex will be set to the last row's index
+	bool GetLastRowIndexAtTime( const float& desiredTime, int32& lastRowIndex );
+
 	// Get the last index of the particles to be spawned at time t
 	// Invalid Index are used to indicate edge cases:
 	// -1 will be returned if no particles have been spawned ( t is smaller than the first particle time )
 	// NumberOfLines will be returned if all particles in the CSV have been spawned ( t is higher than the last particle time )
 	bool GetLastParticleIndexAtTime( const float& time, int32& lastIndex );
+
 	// Returns the previous and next indexes for reading the values of a specified particle at a given time
 	bool GetParticleLineIndexAtTime(const int32& ParticleID, const float& desiredTime, int32& PrevIndex, int32& NextIndex, float& PrevWeight);
 	// Returns the value for a particle at a given time value (linearly interpolated) 
@@ -86,6 +93,9 @@ class HOUDININIAGARA_API UHoudiniCSV : public UObject
 	// Returns the Position Value for a given particle at a given time value (linearly interpolated) 
 	bool GetParticlePositionAtTime(const int32& ParticleID, const float& desiredTime, FVector& Vector);
 
+	// Return a Particle's life for a given time value
+	bool GetParticleLifeAtTime(const int32& ParticleID, const float& DesiredTime, float& Value);
+
 	//-----------------------------------------------------------------------------------------
 	//  MEMBER VARIABLES
 	//-----------------------------------------------------------------------------------------
@@ -125,15 +135,17 @@ class HOUDININIAGARA_API UHoudiniCSV : public UObject
 	// Array containing all the CSV data converted to floats
 	UPROPERTY()
 	TArray<float> FloatCSVData;
-
-	// Array containing the different time values for each particles in the file
-	TArray<float> TimeValues;
-
+	
 	// Array containing the spawn times for each particles in the file
+	UPROPERTY()
 	TArray<float> SpawnTimes;
 
-	// Array containing all the life values for each particles in the file 
-	TArray<float> LifeValues;	
+	// Array containing all the life values for each particles in the file
+	UPROPERTY()
+	TArray<float> LifeValues;
+
+	// Array containing the different time values for each particles in the file
+	TArray<float> TimeValues;
 
 	// Index of the Position values in the buffer
 	UPROPERTY()	    
