Houdini Niagara:

Fixed performance warning for UHoudiniCSV::GetParticleLifeAtTime()

diff --git a/Engine/Plugins/FX/HoudiniNiagara/Source/HoudiniNiagara/Private/HoudiniCSV.cpp b/Engine/Plugins/FX/HoudiniNiagara/Source/HoudiniNiagara/Private/HoudiniCSV.cpp
--- a/Engine/Plugins/FX/HoudiniNiagara/Source/HoudiniNiagara/Private/HoudiniCSV.cpp
+++ b/Engine/Plugins/FX/HoudiniNiagara/Source/HoudiniNiagara/Private/HoudiniCSV.cpp
@@ -619,9 +619,15 @@ bool UHoudiniCSV::GetParticleLifeAtTime( const int32& ParticleID, const float& D
 	}
 
 	if ( DesiredTime < SpawnTimes[ ParticleID ] )
-		return LifeValues[ ParticleID ];
+	{
+		Value = LifeValues[ ParticleID ];
+	}
+	else
+	{
+		Value = LifeValues[ ParticleID ] - ( DesiredTime - SpawnTimes[ ParticleID ] );
+	}
 
-	return LifeValues[ ParticleID ] - ( DesiredTime - SpawnTimes[ ParticleID ] );
+	return true;
 }
 
 bool UHoudiniCSV::GetParticleValueAtTime( const int32& ParticleID, const int32& ColumnIndex, const float& desiredTime, float& Value )
