Niagara - Houdini:

Houdini Niagara Data Interface:
Added GetParticleLifeAtTime for accessing a given particle's life at a given time value.
Fixed GetParticleVectorValueAtTime not bound properly.
Fixed GetRowIndexesForParticleAtTime returning incorrect values when the time value was past the particle's last update.

diff --git a/Engine/Plugins/FX/HoudiniNiagara/Source/HoudiniNiagara/Classes/NiagaraDataInterfaceHoudiniCSV.h b/Engine/Plugins/FX/HoudiniNiagara/Source/HoudiniNiagara/Classes/NiagaraDataInterfaceHoudiniCSV.h
--- a/Engine/Plugins/FX/HoudiniNiagara/Source/HoudiniNiagara/Classes/NiagaraDataInterfaceHoudiniCSV.h
+++ b/Engine/Plugins/FX/HoudiniNiagara/Source/HoudiniNiagara/Classes/NiagaraDataInterfaceHoudiniCSV.h
@@ -119,6 +119,10 @@ public:
 	// Returns the line indexes (previous, next) for reading values for a given particle at a given time
 	template<typename NParamType, typename TimeParamType>
 	void GetRowIndexesForParticleAtTime(FVectorVMContext& Context);
+
+	// Return the life of a given particle at a given time
+	template<typename NParamType, typename TimeParamType>
+	void GetParticleLifeAtTime(FVectorVMContext& Context);
 	
 	//----------------------------------------------------------------------------
 	// GPU / HLSL Functions
diff --git a/Engine/Plugins/FX/HoudiniNiagara/Source/HoudiniNiagara/Private/HoudiniCSV.cpp b/Engine/Plugins/FX/HoudiniNiagara/Source/HoudiniNiagara/Private/HoudiniCSV.cpp
--- a/Engine/Plugins/FX/HoudiniNiagara/Source/HoudiniNiagara/Private/HoudiniCSV.cpp
+++ b/Engine/Plugins/FX/HoudiniNiagara/Source/HoudiniNiagara/Private/HoudiniCSV.cpp
@@ -758,12 +758,14 @@ bool UHoudiniCSV::GetParticleLineIndexAtTime(const int32& ParticleID, const floa
 	if ( PrevIndex < 0 )
 	{
 		PrevWeight = 0.0f;
+		PrevIndex = NextIndex;
 		return true;
 	}
 
 	if ( NextIndex < 0 )
 	{
 		PrevWeight = 1.0f;
+		NextIndex = PrevIndex;
 		return true;
 	}
 
diff --git a/Engine/Plugins/FX/HoudiniNiagara/Source/HoudiniNiagara/Private/NiagaraDataInterfaceHoudiniCSV.cpp b/Engine/Plugins/FX/HoudiniNiagara/Source/HoudiniNiagara/Private/NiagaraDataInterfaceHoudiniCSV.cpp
--- a/Engine/Plugins/FX/HoudiniNiagara/Source/HoudiniNiagara/Private/NiagaraDataInterfaceHoudiniCSV.cpp
+++ b/Engine/Plugins/FX/HoudiniNiagara/Source/HoudiniNiagara/Private/NiagaraDataInterfaceHoudiniCSV.cpp
@@ -320,6 +320,19 @@ void UNiagaraDataInterfaceHoudiniCSV::GetFunctions(TArray<FNiagaraFunctionSignat
 		OutFunctions.Add(Sig);
 	}
 
+	{
+		// GetParticleLifeAtTime
+		FNiagaraFunctionSignature Sig;
+		Sig.Name = TEXT("GetParticleLifeAtTime");
+		Sig.bMemberFunction = true;
+		Sig.bRequiresContext = false;
+		Sig.Inputs.Add(FNiagaraVariable(FNiagaraTypeDefinition(GetClass()), TEXT("CSV")));				// CSV in
+		Sig.Inputs.Add(FNiagaraVariable(FNiagaraTypeDefinition::GetIntDef(), TEXT("N")));				// Point Number In		
+		Sig.Inputs.Add(FNiagaraVariable(FNiagaraTypeDefinition::GetFloatDef(), TEXT("Time")));		    // Time in		
+		Sig.Outputs.Add(FNiagaraVariable(FNiagaraTypeDefinition::GetFloatDef(), TEXT("Value")));		// Float Out
+
+		OutFunctions.Add(Sig);
+	}
 }
 
 // build the shader function HLSL; function name is passed in, as it's defined per-DI; that way, configuration could change
@@ -418,6 +431,7 @@ DEFINE_NDI_RAW_FUNC_BINDER(UNiagaraDataInterfaceHoudiniCSV, GetRowIndexesForPart
 DEFINE_NDI_RAW_FUNC_BINDER(UNiagaraDataInterfaceHoudiniCSV, GetParticlePositionAtTime);
 DEFINE_NDI_RAW_FUNC_BINDER(UNiagaraDataInterfaceHoudiniCSV, GetParticleValueAtTime);
 DEFINE_NDI_RAW_FUNC_BINDER(UNiagaraDataInterfaceHoudiniCSV, GetParticleVectorValueAtTime);
+DEFINE_NDI_RAW_FUNC_BINDER(UNiagaraDataInterfaceHoudiniCSV, GetParticleLifeAtTime);
 void UNiagaraDataInterfaceHoudiniCSV::GetVMExternalFunction(const FVMExternalFunctionBindingInfo& BindingInfo, void* InstanceData, FVMExternalFunction &OutFunc)
 {
     if (BindingInfo.Name == TEXT("GetCSVFloatValue") && BindingInfo.GetNumInputs() == 2 && BindingInfo.GetNumOutputs() == 1)
@@ -476,10 +490,14 @@ void UNiagaraDataInterfaceHoudiniCSV::GetVMExternalFunction(const FVMExternalFun
 	{
 		TNDIParamBinder<0, int32, TNDIParamBinder<1, int32, TNDIParamBinder<2, float, NDI_RAW_FUNC_BINDER(UNiagaraDataInterfaceHoudiniCSV, GetParticleValueAtTime)>>>::Bind(this, BindingInfo, InstanceData, OutFunc);
 	}
-	else if (BindingInfo.Name == TEXT("GetParticleVectorValueAtTime") && BindingInfo.GetNumInputs() == 2 && BindingInfo.GetNumOutputs() == 3)
+	else if (BindingInfo.Name == TEXT("GetParticleVectorValueAtTime") && BindingInfo.GetNumInputs() == 3 && BindingInfo.GetNumOutputs() == 3)
 	{
 		TNDIParamBinder<0, int32, TNDIParamBinder<1, int32, TNDIParamBinder<2, float, NDI_RAW_FUNC_BINDER(UNiagaraDataInterfaceHoudiniCSV, GetParticleVectorValueAtTime)>>>::Bind(this, BindingInfo, InstanceData, OutFunc);
 	}
+	else if (BindingInfo.Name == TEXT("GetParticleLifeAtTime") && BindingInfo.GetNumInputs() == 2 && BindingInfo.GetNumOutputs() == 1)
+	{
+		TNDIParamBinder<0, int32, TNDIParamBinder<1, float, NDI_RAW_FUNC_BINDER(UNiagaraDataInterfaceHoudiniCSV, GetParticleLifeAtTime)>>::Bind(this, BindingInfo, InstanceData, OutFunc);
+	}
     else
     {
 		UE_LOG( LogHoudiniNiagara, Error, 
@@ -907,6 +925,34 @@ void UNiagaraDataInterfaceHoudiniCSV::GetParticleVectorValueAtTime(FVectorVMCont
 	}
 }
 
+template<typename NParamType, typename TimeParamType>
+void UNiagaraDataInterfaceHoudiniCSV::GetParticleLifeAtTime(FVectorVMContext& Context)
+{
+	NParamType NParam(Context);
+	TimeParamType TimeParam(Context);
+
+	FRegisterHandler<float> OutValue(Context);
+
+	for (int32 i = 0; i < Context.NumInstances; ++i)
+	{
+		int32 N = NParam.Get();
+		float time = TimeParam.Get();
+
+		float Value = 0.0f;
+		if ( CSVFile )
+		{
+			CSVFile->GetParticleLifeAtTime(N, time, Value);
+		}
+
+		*OutValue.GetDest() = Value;
+
+		NParam.Advance();
+		TimeParam.Advance();
+
+		OutValue.Advance();
+	}
+}
+
 void UNiagaraDataInterfaceHoudiniCSV::GetNumberOfRowsInCSV(FVectorVMContext& Context)
 {
     FRegisterHandler<int32> OutNumRows(Context);
