Niagara - Houdini:
- Added GetNumberOfColumnsInCSV to the Houdini Data Interface
- Added descriptions to the functions exposed by the DI

diff --git a/Engine/Plugins/FX/HoudiniNiagara/Source/HoudiniNiagara/Classes/NiagaraDataInterfaceHoudiniCSV.h b/Engine/Plugins/FX/HoudiniNiagara/Source/HoudiniNiagara/Classes/NiagaraDataInterfaceHoudiniCSV.h
--- a/Engine/Plugins/FX/HoudiniNiagara/Source/HoudiniNiagara/Classes/NiagaraDataInterfaceHoudiniCSV.h
+++ b/Engine/Plugins/FX/HoudiniNiagara/Source/HoudiniNiagara/Classes/NiagaraDataInterfaceHoudiniCSV.h
@@ -93,6 +93,9 @@ public:
 	// Returns the number of rows found in the CSV file
 	void GetNumberOfRowsInCSV(FVectorVMContext& Context);
 
+	// Returns the number of columns found in the CSV file
+	void GetNumberOfColumnsInCSV(FVectorVMContext& Context);
+
 	// Returns the number of points found in the CSV file
 	void GetNumberOfParticlesInCSV(FVectorVMContext& Context);
 
@@ -140,10 +143,6 @@ protected:
 	UPROPERTY()
 	bool GPUBufferDirty;
 
-	// Last Index used to spawn particles
-	//UPROPERTY()
-	//int32 LastSpawnIndex;
-
 	// Last Spawned Particle ID
 	UPROPERTY()
 	int32 LastSpawnedParticleID;
diff --git a/Engine/Plugins/FX/HoudiniNiagara/Source/HoudiniNiagara/Private/HoudiniCSV.cpp b/Engine/Plugins/FX/HoudiniNiagara/Source/HoudiniNiagara/Private/HoudiniCSV.cpp
--- a/Engine/Plugins/FX/HoudiniNiagara/Source/HoudiniNiagara/Private/HoudiniCSV.cpp
+++ b/Engine/Plugins/FX/HoudiniNiagara/Source/HoudiniNiagara/Private/HoudiniCSV.cpp
@@ -29,7 +29,7 @@
 #include "EditorFramework/AssetImportData.h"
 #endif
 
-#define LOCTEXT_NAMESPACE HOUDINI_NIAGARA_LOCTEXT_NAMESPACE 
+#define LOCTEXT_NAMESPACE "HoudiniNiagaraCSVAsset"
 
 DEFINE_LOG_CATEGORY(LogHoudiniNiagara);
  
@@ -104,6 +104,13 @@ bool UHoudiniCSV::UpdateFromStringArray( TArray<FString>& RawStringArray )
 		return false;
     }
 
+	// Parses the CSV file's title line to update the column indexes of special values we're interested in
+	// Also look for packed vectors in the first line and update the indexes accordingly
+	bool HasPackedVectors = false;
+	if ( !ParseCSVTitleLine( RawStringArray[0], RawStringArray[1], HasPackedVectors ) )
+		return false;
+
+/*
     // Get the number of values per lines via the title line
     FString TitleString = RawStringArray[ 0 ];
     TitleRowArray.Empty();
@@ -271,6 +278,7 @@ bool UHoudiniCSV::UpdateFromStringArray( TArray<FString>& RawStringArray )
 			TEXT("Error while parsing the CSV File. Found %d columns but found %d values in the first line! Some values will have an offset!" ),
 			NumberOfColumns, FirstLineColumns );
     }
+	*/
     
     // Remove the title row now that it's been processed
     RawStringArray.RemoveAt( 0 );
@@ -400,6 +408,179 @@ bool UHoudiniCSV::UpdateFromStringArray( TArray<FString>& RawStringArray )
     return true;
 }
 
+bool UHoudiniCSV::ParseCSVTitleLine( const FString& TitleLine, const FString& FirstLine, bool& HasPackedVectors )
+{
+	// Get the number of values per lines via the title line
+    //FString TitleString = RawStringArray[ 0 ];
+    TitleRowArray.Empty();
+    TitleLine.ParseIntoArray( TitleRowArray, TEXT(",") );
+    NumberOfColumns = TitleRowArray.Num();
+    if ( NumberOfColumns < 1 )
+    {
+		UE_LOG( LogHoudiniNiagara, Error, TEXT( "Could not load the CSV file, error: not enough columns." ) );
+		return false;
+    }
+
+    // Look for the position, normal and time attributes indexes
+    for ( int32 n = 0; n < TitleRowArray.Num(); n++ )
+    {
+		FString CurrentTitle = TitleRowArray[ n ];
+		CurrentTitle.ToLowerInline();
+		CurrentTitle.ReplaceInline( TEXT(" "), TEXT("") );
+
+		if ( CurrentTitle.Equals( TEXT("P"), ESearchCase::IgnoreCase )
+			|| CurrentTitle.Equals( TEXT("Px"), ESearchCase::IgnoreCase )
+			|| CurrentTitle.Equals( TEXT("X"), ESearchCase::IgnoreCase )
+			|| CurrentTitle.Equals(TEXT("pos"), ESearchCase::IgnoreCase ) )
+		{
+			if ( PositionColumnIndex == INDEX_NONE )
+				PositionColumnIndex = n;
+		}
+		else if ( CurrentTitle.Equals( TEXT("N"), ESearchCase::IgnoreCase )
+			|| CurrentTitle.Equals( TEXT("Nx"), ESearchCase::IgnoreCase ) )
+		{
+			if ( NormalColumnIndex == INDEX_NONE )
+				NormalColumnIndex = n;
+		}
+		else if ( ( CurrentTitle.Equals( TEXT("T"), ESearchCase::IgnoreCase ) )
+			|| ( CurrentTitle.Contains( TEXT("time"), ESearchCase::IgnoreCase ) ) )
+		{
+			if ( TimeColumnIndex == INDEX_NONE )
+				TimeColumnIndex = n;
+		}
+		else if ( ( CurrentTitle.Equals( TEXT("#"), ESearchCase::IgnoreCase ) )
+			|| ( CurrentTitle.Equals( TEXT("id"), ESearchCase::IgnoreCase ) ) )
+		{
+			if ( IDColumnIndex == INDEX_NONE )
+				IDColumnIndex = n;
+		}
+		else if ( CurrentTitle.Equals( TEXT("alive"), ESearchCase::IgnoreCase ) )
+		{
+			if (AliveColumnIndex == INDEX_NONE)
+				AliveColumnIndex = n;
+		}
+		else if ( CurrentTitle.Equals( TEXT("life"), ESearchCase::IgnoreCase ) )
+		{
+			if ( LifeColumnIndex == INDEX_NONE)
+				LifeColumnIndex = n;
+		}
+    }
+
+	 // Read the first line of the CSV file, and look for packed vectors value (X,Y,Z)
+    // We'll have to expand them in the title row to match the parsed data
+    int32 FoundPackedVectorCharIndex = 0;
+    HasPackedVectors = false;
+    while ( FoundPackedVectorCharIndex != INDEX_NONE )
+    {
+		// Try to find ( in the line
+		FoundPackedVectorCharIndex = FirstLine.Find( TEXT("("), ESearchCase::IgnoreCase, ESearchDir::FromStart, FoundPackedVectorCharIndex );
+		if ( FoundPackedVectorCharIndex == INDEX_NONE )
+			break;
+
+		// We want to know which column this char belong to
+		int32 FoundPackedVectorColumnIndex = INDEX_NONE;
+		{
+			// Chop the first line up to the found character
+			FString FirstLineLeft = FirstLine.Left( FoundPackedVectorCharIndex );
+
+			// ReplaceInLine returns the number of occurences of ",", that's what we want! 
+			FoundPackedVectorColumnIndex = FirstLineLeft.ReplaceInline(TEXT(","), TEXT(""));
+		}
+
+		if ( !TitleRowArray.IsValidIndex( FoundPackedVectorColumnIndex ) )
+		{
+			UE_LOG( LogHoudiniNiagara, Warning,
+			TEXT( "Error while parsing the CSV File. Couldn't unpack vector found at character %d in the first line!" ),
+			FoundPackedVectorCharIndex + 1 );
+			continue;
+		}
+
+		// We found a packed vector, get its size
+		int32 FoundVectorSize = 0;
+		{
+			// Extract the vector string
+			int32 FoundPackedVectorEndCharIndex = FirstLine.Find( TEXT(")"), ESearchCase::IgnoreCase, ESearchDir::FromStart, FoundPackedVectorCharIndex );
+			FString VectorString = FirstLine.Mid( FoundPackedVectorCharIndex + 1, FoundPackedVectorEndCharIndex - FoundPackedVectorCharIndex - 1 );
+
+			// Use ReplaceInLine to count the number of , to get the vector's size!
+			FoundVectorSize = VectorString.ReplaceInline( TEXT(","), TEXT("") ) + 1;
+		}
+
+		if ( FoundVectorSize < 2 )
+			continue;
+
+		// Increment the number of columns
+		NumberOfColumns += ( FoundVectorSize - 1 );
+
+		// Expand TitleRowArray
+		if ( ( FoundPackedVectorColumnIndex == PositionColumnIndex ) && ( FoundVectorSize == 3 ) )
+		{
+			// Expand P to Px,Py,Pz
+			TitleRowArray[ PositionColumnIndex ] = TEXT("Px");
+			TitleRowArray.Insert( TEXT( "Py" ), PositionColumnIndex + 1 );
+			TitleRowArray.Insert( TEXT( "Pz" ), PositionColumnIndex + 2 );
+		}
+		else if ( ( FoundPackedVectorColumnIndex == NormalColumnIndex ) && ( FoundVectorSize == 3 ) )
+		{
+			// Expand N to Nx,Ny,Nz
+			TitleRowArray[ NormalColumnIndex ] = TEXT("Nx");
+			TitleRowArray.Insert( TEXT("Ny"), NormalColumnIndex + 1 );
+			TitleRowArray.Insert( TEXT("Nz"), NormalColumnIndex + 2 );
+		}
+		else
+		{
+			// Expand the vector's title from V to V, V1, V2, V3 ...
+			FString FoundPackedVectortTitle = TitleRowArray[ FoundPackedVectorColumnIndex ];
+			for ( int32 n = 1; n < FoundVectorSize; n++ )
+			{
+				FString CurrentTitle = FoundPackedVectortTitle + FString::FromInt( n );
+				TitleRowArray.Insert( CurrentTitle, FoundPackedVectorColumnIndex + n );
+			}
+		}
+
+		// Eventually offset the stored index
+		if ( PositionColumnIndex != INDEX_NONE && ( PositionColumnIndex > FoundPackedVectorColumnIndex) )
+			PositionColumnIndex += FoundVectorSize - 1;
+
+		if ( NormalColumnIndex != INDEX_NONE && ( NormalColumnIndex > FoundPackedVectorColumnIndex) )
+			NormalColumnIndex += FoundVectorSize - 1;
+
+		if ( TimeColumnIndex != INDEX_NONE && ( TimeColumnIndex > FoundPackedVectorColumnIndex) )
+			TimeColumnIndex += FoundVectorSize - 1;
+
+		if ( IDColumnIndex != INDEX_NONE && ( IDColumnIndex > FoundPackedVectorColumnIndex ) )
+			IDColumnIndex += FoundVectorSize - 1;
+
+		if ( AliveColumnIndex != INDEX_NONE && ( AliveColumnIndex > FoundPackedVectorColumnIndex ) )
+			AliveColumnIndex += FoundVectorSize - 1;
+
+		if ( LifeColumnIndex != INDEX_NONE && ( LifeColumnIndex > FoundPackedVectorColumnIndex ) )
+			LifeColumnIndex += FoundVectorSize - 1;
+
+		HasPackedVectors = true;
+		FoundPackedVectorCharIndex++;
+    }
+
+    // For sanity, Check that the number of columns matches the title row and the first line
+    {
+		// Check the title row
+		if ( NumberOfColumns != TitleRowArray.Num() )
+			UE_LOG( LogHoudiniNiagara, Error,
+			TEXT( "Error while parsing the CSV File. Found %d columns but the Title string has %d values! Some values will have an offset!" ),
+			NumberOfColumns, TitleRowArray.Num() );
+
+		// Use ReplaceInLine to count the number of columns in the first line and make sure it's correct
+		FString FirstLineCopy = FirstLine;
+		int32 FirstLineColumns = FirstLineCopy.ReplaceInline( TEXT(","), TEXT("") ) + 1;
+		if ( NumberOfColumns != FirstLineColumns )
+			UE_LOG( LogHoudiniNiagara, Error,
+			TEXT("Error while parsing the CSV File. Found %d columns but found %d values in the first line! Some values will have an offset!" ),
+			NumberOfColumns, FirstLineColumns );
+    }
+
+	return true;
+}
+
 
 // Returns the float value at a given point in the CSV file
 bool UHoudiniCSV::GetCSVFloatValue( const int32& lineIndex, const int32& colIndex, float& value )
@@ -511,12 +692,18 @@ int32 UHoudiniCSV::GetNumberOfParticlesInCSV()
     return NumberOfLines;
 }
 
-// Returns the number of points found in the CSV file
+// Returns the number of lines found in the CSV file
 int32 UHoudiniCSV::GetNumberOfLinesInCSV()
 {
 	return NumberOfLines;
 }
 
+// Returns the number of columns found in the CSV file
+int32 UHoudiniCSV::GetNumberOfColumnsInCSV()
+{
+	return NumberOfColumns;
+}
+
 // Get the last row index for a given time value (the row with a time smaller or equal to desiredTime)
 // If the CSV file doesn't have time informations, returns false and set the LastRowIndex to the last line in the file
 // If desiredTime is smaller than the time value in the first row, LastRowIndex will be set to -1
diff --git a/Engine/Plugins/FX/HoudiniNiagara/Source/HoudiniNiagara/Private/NiagaraDataInterfaceHoudiniCSV.cpp b/Engine/Plugins/FX/HoudiniNiagara/Source/HoudiniNiagara/Private/NiagaraDataInterfaceHoudiniCSV.cpp
--- a/Engine/Plugins/FX/HoudiniNiagara/Source/HoudiniNiagara/Private/NiagaraDataInterfaceHoudiniCSV.cpp
+++ b/Engine/Plugins/FX/HoudiniNiagara/Source/HoudiniNiagara/Private/NiagaraDataInterfaceHoudiniCSV.cpp
@@ -25,14 +25,13 @@
 #include "NiagaraTypes.h"
 #include "Misc/FileHelper.h"
 
-#define LOCTEXT_NAMESPACE HOUDINI_NIAGARA_LOCTEXT_NAMESPACE 
+#define LOCTEXT_NAMESPACE "HoudiniNiagaraCSVDataInterface"  
 
 
 UNiagaraDataInterfaceHoudiniCSV::UNiagaraDataInterfaceHoudiniCSV(FObjectInitializer const& ObjectInitializer)
 	: Super(ObjectInitializer)
 {
     CSVFile = nullptr;
-    //LastSpawnIndex = -1;
 	LastSpawnedParticleID = -1;
 }
 
@@ -46,7 +45,6 @@ void UNiagaraDataInterfaceHoudiniCSV::PostInitProperties()
     }
 
     GPUBufferDirty = true;
-    //LastSpawnIndex = -1;
 	LastSpawnedParticleID = -1;
 }
 
@@ -54,7 +52,6 @@ void UNiagaraDataInterfaceHoudiniCSV::PostLoad()
 {
     Super::PostLoad();
     GPUBufferDirty = true;
-    //LastSpawnIndex = -1;
 	LastSpawnedParticleID = -1;
 }
 
@@ -70,7 +67,6 @@ void UNiagaraDataInterfaceHoudiniCSV::PostEditChangeProperty(struct FPropertyCha
 		if (CSVFile)
 		{
 			GPUBufferDirty = true;
-			//LastSpawnIndex = -1;
 			LastSpawnedParticleID = -1;
 		}
     }
@@ -122,7 +118,10 @@ void UNiagaraDataInterfaceHoudiniCSV::GetFunctions(TArray<FNiagaraFunctionSignat
 		Sig.Inputs.Add(FNiagaraVariable(FNiagaraTypeDefinition::GetIntDef(), TEXT("Col")));			// Col Index In
 		Sig.Outputs.Add(FNiagaraVariable(FNiagaraTypeDefinition::GetFloatDef(), TEXT("Value")));	// Float Out
 
-		OutFunctions.Add(Sig);
+		Sig.SetDescription( LOCTEXT( "DataInterfaceHoudini_GetCSVFloatValue",
+			"Returns the float value in the CSV file for a given Row and Column.\n" ) );
+
+		OutFunctions.Add( Sig );
     }
 
 	/*
@@ -137,6 +136,9 @@ void UNiagaraDataInterfaceHoudiniCSV::GetFunctions(TArray<FNiagaraFunctionSignat
 		Sig.Inputs.Add(FNiagaraVariable(FNiagaraTypeDefinition::GetStringDef(), TEXT("ColTitle")));	// Col Title In
 		Sig.Outputs.Add(FNiagaraVariable(FNiagaraTypeDefinition::GetFloatDef(), TEXT("Value")));	// Float Out
 
+		Sig.SetDescription( LOCTEXT( "DataInterfaceHoudini_GetCSVFloatValueByString",
+		"Returns the float value in the CSV file for a given Row, in the column corresponding to the ColTitle string.\n" ) );
+
 		OutFunctions.Add(Sig);
     }
 	*/
@@ -152,6 +154,9 @@ void UNiagaraDataInterfaceHoudiniCSV::GetFunctions(TArray<FNiagaraFunctionSignat
 		Sig.Inputs.Add(FNiagaraVariable(FNiagaraTypeDefinition::GetIntDef(), TEXT("Col")));			// Col Index In
 		Sig.Outputs.Add(FNiagaraVariable(FNiagaraTypeDefinition::GetVec3Def(), TEXT("Value")));		// Vector3 Out
 
+		Sig.SetDescription( LOCTEXT( "DataInterfaceHoudini_GetCSVVectorValue",
+			"Returns a Vector3 in the CSV file for a given Row and Column.\nThe returned Vector is converted from Houdini's coordinate system to Unreal's." ) );
+
 		OutFunctions.Add(Sig);
 	}
 
@@ -165,6 +170,9 @@ void UNiagaraDataInterfaceHoudiniCSV::GetFunctions(TArray<FNiagaraFunctionSignat
 		Sig.Inputs.Add(FNiagaraVariable(FNiagaraTypeDefinition::GetIntDef(), TEXT("Row")));			// Row Index In
 		Sig.Outputs.Add(FNiagaraVariable(FNiagaraTypeDefinition::GetVec3Def(), TEXT("Position")));	// Vector3 Out
 
+		Sig.SetDescription( LOCTEXT( "DataInterfaceHoudini_GetCSVPosition",
+			"Helper function returning the position value for a given Row in the CSV file.\nThe returned Position vector is converted from Houdini's coordinate system to Unreal's." ) );
+
 		OutFunctions.Add(Sig);
     }
 
@@ -179,6 +187,9 @@ void UNiagaraDataInterfaceHoudiniCSV::GetFunctions(TArray<FNiagaraFunctionSignat
 		Sig.Outputs.Add(FNiagaraVariable(FNiagaraTypeDefinition::GetVec3Def(), TEXT("Position")));	// Vector3 Out
 		Sig.Outputs.Add(FNiagaraVariable(FNiagaraTypeDefinition::GetFloatDef(), TEXT("Time")));		// float Out
 
+		Sig.SetDescription( LOCTEXT( "DataInterfaceHoudini_GetCSVPositionAndTime",
+			"Helper function returning the position and time values for a given Row in the CSV file.\nThe returned Position vector is converted from Houdini's coordinate system to Unreal's." ) );
+
 		OutFunctions.Add(Sig);
     }
 
@@ -192,6 +203,9 @@ void UNiagaraDataInterfaceHoudiniCSV::GetFunctions(TArray<FNiagaraFunctionSignat
 		Sig.Inputs.Add(FNiagaraVariable(FNiagaraTypeDefinition::GetIntDef(), TEXT("Row")));			// Row Index In
 		Sig.Outputs.Add(FNiagaraVariable(FNiagaraTypeDefinition::GetVec3Def(), TEXT("Normal")));	// Vector3 Out
 
+		Sig.SetDescription( LOCTEXT( "DataInterfaceHoudini_GetCSVNormal",
+			"Helper function returning the normal value for a given Row in the CSV file.\nThe returned Normal vector is converted from Houdini's coordinate system to Unreal's." ) );
+
 		OutFunctions.Add(Sig);
     }
 
@@ -205,6 +219,9 @@ void UNiagaraDataInterfaceHoudiniCSV::GetFunctions(TArray<FNiagaraFunctionSignat
 		Sig.Inputs.Add(FNiagaraVariable(FNiagaraTypeDefinition::GetIntDef(), TEXT("Row")));			// Row Index In
 		Sig.Outputs.Add(FNiagaraVariable(FNiagaraTypeDefinition::GetFloatDef(), TEXT("Time")));		// Float Out
 
+		Sig.SetDescription( LOCTEXT("DataInterfaceHoudini_GetCSVTime",
+			"Helper function returning the time value for a given Row in the CSV file.\n") );
+
 		OutFunctions.Add(Sig);
     }
 
@@ -217,6 +234,9 @@ void UNiagaraDataInterfaceHoudiniCSV::GetFunctions(TArray<FNiagaraFunctionSignat
 		Sig.Inputs.Add(FNiagaraVariable(FNiagaraTypeDefinition(GetClass()), TEXT("CSV")));					// CSV in
 		Sig.Outputs.Add(FNiagaraVariable(FNiagaraTypeDefinition::GetIntDef(), TEXT("NumberOfParticles")));  // Int Out
 
+		Sig.SetDescription( LOCTEXT( "DataInterfaceHoudini_GetNumberOfParticlesInCSV",
+			"Returns the number of particles (with different id values) in the CSV file.\n" ) );
+
 		OutFunctions.Add(Sig);
     }
 
@@ -226,8 +246,26 @@ void UNiagaraDataInterfaceHoudiniCSV::GetFunctions(TArray<FNiagaraFunctionSignat
 		Sig.Name = TEXT("GetNumberOfRowsInCSV");
 		Sig.bMemberFunction = true;
 		Sig.bRequiresContext = false;
-		Sig.Inputs.Add(FNiagaraVariable(FNiagaraTypeDefinition(GetClass()), TEXT("CSV")));					// CSV in
-		Sig.Outputs.Add(FNiagaraVariable(FNiagaraTypeDefinition::GetIntDef(), TEXT("NumberOfRows")));		// Int Out
+		Sig.Inputs.Add( FNiagaraVariable( FNiagaraTypeDefinition( GetClass()), TEXT("CSV") ) );					// CSV in
+		Sig.Outputs.Add( FNiagaraVariable( FNiagaraTypeDefinition::GetIntDef(), TEXT("NumberOfRows") ) );		// Int Out
+		
+		Sig.SetDescription( LOCTEXT( "DataInterfaceHoudini_GetNumberOfRowsInCSV",
+			"Returns the number of rows in the CSV file.\nOnly the number of value rows is returned, the first \"Title\" Row is ignored." ) );
+
+		OutFunctions.Add(Sig);
+	}
+
+		{
+		// GetNumberOfColumnsInCSV
+		FNiagaraFunctionSignature Sig;
+		Sig.Name = TEXT("GetNumberOfColumnsInCSV");
+		Sig.bMemberFunction = true;
+		Sig.bRequiresContext = false;
+		Sig.Inputs.Add( FNiagaraVariable( FNiagaraTypeDefinition( GetClass()), TEXT("CSV") ) );						// CSV in
+		Sig.Outputs.Add( FNiagaraVariable( FNiagaraTypeDefinition::GetIntDef(), TEXT("NumberOfColumns") ) );		// Int Out
+		
+		Sig.SetDescription( LOCTEXT( "DataInterfaceHoudini_GetNumberOfColumnsInCSV",
+			"Returns the number of columns in the CSV file." ) );
 
 		OutFunctions.Add(Sig);
 	}
@@ -238,9 +276,12 @@ void UNiagaraDataInterfaceHoudiniCSV::GetFunctions(TArray<FNiagaraFunctionSignat
 		Sig.Name = TEXT("GetLastRowIndexAtTime");
 		Sig.bMemberFunction = true;
 		Sig.bRequiresContext = false;
-		Sig.Inputs.Add(FNiagaraVariable(FNiagaraTypeDefinition(GetClass()), TEXT("CSV")));					// CSV in
-		Sig.Inputs.Add(FNiagaraVariable(FNiagaraTypeDefinition::GetFloatDef(), TEXT("Time")));				// Time in
-		Sig.Outputs.Add(FNiagaraVariable(FNiagaraTypeDefinition::GetIntDef(), TEXT("LastRowIndex")));	    // Int Out
+		Sig.Inputs.Add(FNiagaraVariable( FNiagaraTypeDefinition(GetClass()), TEXT("CSV") ) );					// CSV in
+		Sig.Inputs.Add(FNiagaraVariable( FNiagaraTypeDefinition::GetFloatDef(), TEXT("Time") ) );				// Time in
+		Sig.Outputs.Add(FNiagaraVariable( FNiagaraTypeDefinition::GetIntDef(), TEXT("LastRowIndex") ) );	    // Int Out
+
+		Sig.SetDescription( LOCTEXT( "DataInterfaceHoudini_GetLastRowIndexAtTime",
+			"Returns the index of the last row in the CSV file that has a time value lesser or equal to the Time parameter." ) );
 
 		OutFunctions.Add(Sig);
     }
@@ -251,11 +292,14 @@ void UNiagaraDataInterfaceHoudiniCSV::GetFunctions(TArray<FNiagaraFunctionSignat
 		Sig.Name = TEXT("GetParticleIndexesToSpawnAtTime");
 		Sig.bMemberFunction = true;
 		Sig.bRequiresContext = false;
-		Sig.Inputs.Add(FNiagaraVariable(FNiagaraTypeDefinition(GetClass()), TEXT("CSV")));				// CSV in
-		Sig.Inputs.Add(FNiagaraVariable(FNiagaraTypeDefinition::GetFloatDef(), TEXT("Time")));		    // Time in
-		Sig.Outputs.Add(FNiagaraVariable(FNiagaraTypeDefinition::GetIntDef(), TEXT("MinIndex")));	    // Int Out
-		Sig.Outputs.Add(FNiagaraVariable(FNiagaraTypeDefinition::GetIntDef(), TEXT("MaxIndex")));	    // Int Out
-		Sig.Outputs.Add(FNiagaraVariable(FNiagaraTypeDefinition::GetIntDef(), TEXT("Count")));		    // Int Out
+		Sig.Inputs.Add(FNiagaraVariable( FNiagaraTypeDefinition(GetClass()), TEXT("CSV") ) );				// CSV in
+		Sig.Inputs.Add(FNiagaraVariable( FNiagaraTypeDefinition::GetFloatDef(), TEXT("Time") ) );		    // Time in
+		Sig.Outputs.Add(FNiagaraVariable( FNiagaraTypeDefinition::GetIntDef(), TEXT("MinIndex") ) );	    // Int Out
+		Sig.Outputs.Add(FNiagaraVariable( FNiagaraTypeDefinition::GetIntDef(), TEXT("MaxIndex") ) );	    // Int Out
+		Sig.Outputs.Add(FNiagaraVariable( FNiagaraTypeDefinition::GetIntDef(), TEXT("Count") ) );		    // Int Out
+
+		Sig.SetDescription( LOCTEXT( "DataInterfaceHoudini_GetParticleIndexesToSpawnAtTime",
+			"Returns the count and particle IDs of the particles that should spawn at a given time value." ) );
 
 		OutFunctions.Add(Sig);
     }
@@ -273,6 +317,9 @@ void UNiagaraDataInterfaceHoudiniCSV::GetFunctions(TArray<FNiagaraFunctionSignat
 		Sig.Outputs.Add(FNiagaraVariable(FNiagaraTypeDefinition::GetIntDef(), TEXT("NextRow")));			// Int Out
 		Sig.Outputs.Add(FNiagaraVariable(FNiagaraTypeDefinition::GetFloatDef(), TEXT("PrevWeight")));		// Float Out
 
+		Sig.SetDescription( LOCTEXT( "DataInterfaceHoudini_GetRowIndexesForParticleAtTime",
+			"Returns the row indexes for a given particle at a given time.\nThe previous row, next row and weight can then be used to Lerp between values.") );
+
 		OutFunctions.Add(Sig);
 	}
 
@@ -287,6 +334,9 @@ void UNiagaraDataInterfaceHoudiniCSV::GetFunctions(TArray<FNiagaraFunctionSignat
 		Sig.Inputs.Add(FNiagaraVariable(FNiagaraTypeDefinition::GetFloatDef(), TEXT("Time")));		    // Time in
 		Sig.Outputs.Add(FNiagaraVariable(FNiagaraTypeDefinition::GetVec3Def(), TEXT("Position")));		// Vector3 Out
 
+		Sig.SetDescription( LOCTEXT( "DataInterfaceHoudini_GetParticlePositionAtTime",
+			"Helper function returning the linearly interpolated position for a given particle at a given time.\nThe returned Position vector is converted from Houdini's coordinate system to Unreal's.") );
+
 		OutFunctions.Add(Sig);
 	}
 
@@ -302,6 +352,9 @@ void UNiagaraDataInterfaceHoudiniCSV::GetFunctions(TArray<FNiagaraFunctionSignat
 		Sig.Inputs.Add(FNiagaraVariable(FNiagaraTypeDefinition::GetFloatDef(), TEXT("Time")));		    // Time in		
 		Sig.Outputs.Add(FNiagaraVariable(FNiagaraTypeDefinition::GetFloatDef(), TEXT("Value")));		// Float Out
 
+		Sig.SetDescription( LOCTEXT( "DataInterfaceHoudini_GetParticleValueAtTime",
+			"Returns the linearly interpolated value in the specified column for a given particle at a given time." ) );
+
 		OutFunctions.Add(Sig);
 	}
 
@@ -317,6 +370,9 @@ void UNiagaraDataInterfaceHoudiniCSV::GetFunctions(TArray<FNiagaraFunctionSignat
 		Sig.Inputs.Add(FNiagaraVariable(FNiagaraTypeDefinition::GetFloatDef(), TEXT("Time")));		    // Time in		
 		Sig.Outputs.Add(FNiagaraVariable(FNiagaraTypeDefinition::GetVec3Def(), TEXT("Value")));			// Vector3 Out
 
+		Sig.SetDescription( LOCTEXT( "DataInterfaceHoudini_GetParticleVectorValueAtTime",
+			"Helper function returning the linearly interpolated Vector value in the specified column for a given particle at a given time.\nThe returned Vector is converted from Houdini's coordinate system to Unreal's." ) );
+
 		OutFunctions.Add(Sig);
 	}
 
@@ -331,6 +387,9 @@ void UNiagaraDataInterfaceHoudiniCSV::GetFunctions(TArray<FNiagaraFunctionSignat
 		Sig.Inputs.Add(FNiagaraVariable(FNiagaraTypeDefinition::GetFloatDef(), TEXT("Time")));		    // Time in		
 		Sig.Outputs.Add(FNiagaraVariable(FNiagaraTypeDefinition::GetFloatDef(), TEXT("Value")));		// Float Out
 
+		Sig.SetDescription( LOCTEXT( "DataInterfaceHoudini_GetParticleLifeAtTime",
+			"Helper function returning the remaining life for a given particle in the CSV file at a given time." ) );
+
 		OutFunctions.Add(Sig);
 	}
 }
@@ -470,6 +529,10 @@ void UNiagaraDataInterfaceHoudiniCSV::GetVMExternalFunction(const FVMExternalFun
 	{
 		OutFunc = FVMExternalFunction::CreateUObject(this, &UNiagaraDataInterfaceHoudiniCSV::GetNumberOfRowsInCSV);
 	}
+	else if (BindingInfo.Name == TEXT("GetNumberOfColumnsInCSV") && BindingInfo.GetNumInputs() == 0 && BindingInfo.GetNumOutputs() == 1)
+	{
+		OutFunc = FVMExternalFunction::CreateUObject(this, &UNiagaraDataInterfaceHoudiniCSV::GetNumberOfColumnsInCSV);
+	}
     else if (BindingInfo.Name == TEXT("GetLastRowIndexAtTime") && BindingInfo.GetNumInputs() == 1 && BindingInfo.GetNumOutputs() == 1)
     {
 		TNDIParamBinder<0, float, NDI_RAW_FUNC_BINDER(UNiagaraDataInterfaceHoudiniCSV, GetLastRowIndexAtTime)>::Bind(this, BindingInfo, InstanceData, OutFunc);
@@ -707,7 +770,7 @@ void UNiagaraDataInterfaceHoudiniCSV::GetParticleIndexesToSpawnAtTime( FVectorVM
 			else
 			{
 				// The CSV file has time informations
-				// First, detect if we need to reset LastSpawnIndex (after a loop)
+				// First, detect if we need to reset LastSpawnParticleID (after a loop of the emitter)
 				if ( value < LastSpawnedParticleID )
 					LastSpawnedParticleID = -1;
 
@@ -960,6 +1023,13 @@ void UNiagaraDataInterfaceHoudiniCSV::GetNumberOfRowsInCSV(FVectorVMContext& Con
 	OutNumRows.Advance();
 }
 
+void UNiagaraDataInterfaceHoudiniCSV::GetNumberOfColumnsInCSV(FVectorVMContext& Context)
+{
+	FRegisterHandler<int32> OutNumCols(Context);
+	*OutNumCols.GetDest() = CSVFile ? CSVFile->GetNumberOfColumnsInCSV() : 0;
+	OutNumCols.Advance();
+}
+
 void UNiagaraDataInterfaceHoudiniCSV::GetNumberOfParticlesInCSV(FVectorVMContext& Context)
 {
 	FRegisterHandler<int32> OutNumParticles(Context);
diff --git a/Engine/Plugins/FX/HoudiniNiagara/Source/HoudiniNiagara/Public/HoudiniCSV.h b/Engine/Plugins/FX/HoudiniNiagara/Source/HoudiniNiagara/Public/HoudiniCSV.h
--- a/Engine/Plugins/FX/HoudiniNiagara/Source/HoudiniNiagara/Public/HoudiniCSV.h
+++ b/Engine/Plugins/FX/HoudiniNiagara/Source/HoudiniNiagara/Public/HoudiniCSV.h
@@ -51,6 +51,9 @@ class HOUDININIAGARA_API UHoudiniCSV : public UObject
 	// Returns the number of points found in the CSV file
 	int32 GetNumberOfLinesInCSV();
 
+	// Returns the number of columns found in the CSV file
+	int32 GetNumberOfColumnsInCSV();
+
 	// Returns the column index for a given string
 	bool GetColumnIndexFromString(const FString& ColumnTitle, int32& ColumnIndex);
 
@@ -126,6 +129,13 @@ class HOUDININIAGARA_API UHoudiniCSV : public UObject
 	virtual void PostInitProperties() override;
 #endif
 
+	protected:
+
+	// Parses the CSV file's title line to update the column indexes of special values we're interested in
+	// Also look for packed vectors in the first line and update the indexes accordingly
+	bool ParseCSVTitleLine( const FString& CSVTitleString, const FString& FirstLine, bool& HasPackedVectors );
+
+
     private:
 
 	// Array containing the Raw String data
