Niagara - Houdini:

Houdini CSV Asset:
- As the DI expects the values to be sorted by time, if it's not the case, the CSV importer will sort them on import.
- As the DI spawning functions relies on the particle IDs starting at zero and increment, the CSV importer will fix the particles IDs on import if it's not the case.

diff --git a/Engine/Plugins/FX/HoudiniNiagara/Source/HoudiniNiagara/Private/HoudiniCSV.cpp b/Engine/Plugins/FX/HoudiniNiagara/Source/HoudiniNiagara/Private/HoudiniCSV.cpp
--- a/Engine/Plugins/FX/HoudiniNiagara/Source/HoudiniNiagara/Private/HoudiniCSV.cpp
+++ b/Engine/Plugins/FX/HoudiniNiagara/Source/HoudiniNiagara/Private/HoudiniCSV.cpp
@@ -24,6 +24,7 @@
 #include "HoudiniCSV.h"
 #include "Misc/FileHelper.h"
 #include "Misc/Paths.h"
+#include "Math/NumericLimits.h"
 
 #if WITH_EDITOR
 #include "EditorFramework/AssetImportData.h"
@@ -32,6 +33,45 @@
 #define LOCTEXT_NAMESPACE "HoudiniNiagaraCSVAsset"
 
 DEFINE_LOG_CATEGORY(LogHoudiniNiagara);
+
+
+struct FHoudiniCSVSortPredicate
+{
+	FHoudiniCSVSortPredicate(const int32 &InTimeCol, const int32 &InIDCol )
+		: TimeColumnIndex( InTimeCol ), IDColumnIndex( InIDCol )
+	{}
+
+	bool operator()( const TArray<FString>& A, const TArray<FString>& B ) const
+	{
+		float ATime = TNumericLimits< float >::Lowest();
+		if ( A.IsValidIndex( TimeColumnIndex ) )
+			ATime = FCString::Atof( *A[ TimeColumnIndex ] );
+
+		float BTime = TNumericLimits< float >::Lowest();
+		if ( B.IsValidIndex( TimeColumnIndex ) )
+			BTime = FCString::Atof( *B[ TimeColumnIndex ] );
+
+		if ( ATime != BTime )
+		{
+			return ATime < BTime;
+		}
+		else
+		{
+			float AID = TNumericLimits< float >::Lowest();
+			if ( A.IsValidIndex( IDColumnIndex ) )
+				AID = FCString::Atof( *A[ IDColumnIndex ]);
+
+			float BID = TNumericLimits< float >::Lowest();
+			if ( B.IsValidIndex( IDColumnIndex ) )
+				BID = FCString::Atof( *B[ IDColumnIndex ] );
+
+			return AID <= BID;
+		}
+	}
+
+	int32 TimeColumnIndex;
+	int32 IDColumnIndex;
+};
  
 UHoudiniCSV::UHoudiniCSV( const FObjectInitializer& ObjectInitializer )
 	: Super( ObjectInitializer ),
@@ -109,179 +149,72 @@ bool UHoudiniCSV::UpdateFromStringArray( TArray<FString>& RawStringArray )
 	bool HasPackedVectors = false;
 	if ( !ParseCSVTitleLine( RawStringArray[0], RawStringArray[1], HasPackedVectors ) )
 		return false;
+    
+    // Remove the title row now that it's been processed
+    RawStringArray.RemoveAt( 0 );
 
-/*
-    // Get the number of values per lines via the title line
-    FString TitleString = RawStringArray[ 0 ];
-    TitleRowArray.Empty();
-    TitleString.ParseIntoArray( TitleRowArray, TEXT(",") );
-    NumberOfColumns = TitleRowArray.Num();
-    if ( NumberOfColumns < 1 )
-    {
-		UE_LOG( LogHoudiniNiagara, Error, TEXT( "Could not load the CSV file, error: not enough columns." ) );
-		return false;
-    }
-
-    // Look for the position, normal and time attributes indexes
-    for ( int32 n = 0; n < TitleRowArray.Num(); n++ )
-    {
-		FString CurrentTitle = TitleRowArray[ n ];
-		CurrentTitle.ToLowerInline();
-		CurrentTitle.ReplaceInline( TEXT(" "), TEXT("") );
-
-		if ( CurrentTitle.Equals( TEXT("P"), ESearchCase::IgnoreCase )
-			|| CurrentTitle.Equals( TEXT("Px"), ESearchCase::IgnoreCase )
-			|| CurrentTitle.Equals( TEXT("X"), ESearchCase::IgnoreCase )
-			|| CurrentTitle.Equals(TEXT("pos"), ESearchCase::IgnoreCase ) )
-		{
-			if ( PositionColumnIndex == INDEX_NONE )
-				PositionColumnIndex = n;
-		}
-		else if ( CurrentTitle.Equals( TEXT("N"), ESearchCase::IgnoreCase )
-			|| CurrentTitle.Equals( TEXT("Nx"), ESearchCase::IgnoreCase ) )
-		{
-			if ( NormalColumnIndex == INDEX_NONE )
-				NormalColumnIndex = n;
-		}
-		else if ( ( CurrentTitle.Equals( TEXT("T"), ESearchCase::IgnoreCase ) )
-			|| ( CurrentTitle.Contains( TEXT("time"), ESearchCase::IgnoreCase ) ) )
-		{
-			if ( TimeColumnIndex == INDEX_NONE )
-				TimeColumnIndex = n;
-		}
-		else if ( ( CurrentTitle.Equals( TEXT("#"), ESearchCase::IgnoreCase ) )
-			|| ( CurrentTitle.Equals( TEXT("id"), ESearchCase::IgnoreCase ) ) )
-		{
-			if ( IDColumnIndex == INDEX_NONE )
-				IDColumnIndex = n;
-		}
-		else if ( CurrentTitle.Equals( TEXT("alive"), ESearchCase::IgnoreCase ) )
-		{
-			if (AliveColumnIndex == INDEX_NONE)
-				AliveColumnIndex = n;
-		}
-		else if ( CurrentTitle.Equals( TEXT("life"), ESearchCase::IgnoreCase ) )
+	// Parses each string of the csv file to a string array
+	TArray< TArray< FString > > ParsedStringArrays;
+	ParsedStringArrays.SetNum( NumberOfLines );
+	for ( int32 rowIdx = 0; rowIdx < NumberOfLines; rowIdx++ )
+	{
+		// Get the current line
+		FString CurrentLine = RawStringArray[ rowIdx ];
+		if ( HasPackedVectors )
 		{
-			if ( LifeColumnIndex == INDEX_NONE)
-				LifeColumnIndex = n;
+			// Clean up the packing characters: ()" from the line so it can be parsed properly
+			CurrentLine.ReplaceInline( TEXT("("), TEXT("") );
+			CurrentLine.ReplaceInline( TEXT(")"), TEXT("") );
+			CurrentLine.ReplaceInline( TEXT("\""), TEXT("") );
 		}
-    }
-
-    // Read the first line of the CSV file, and look for packed vectors value (X,Y,Z)
-    // We'll have to expand them in the title row to match the parsed data
-    FString FirstLine = RawStringArray[ 1 ];
-    int32 FoundPackedVectorCharIndex = 0;
-    bool HasPackedVectors = false;
-    while ( FoundPackedVectorCharIndex != INDEX_NONE )
-    {
-		// Try to find ( in the line
-		FoundPackedVectorCharIndex = FirstLine.Find( TEXT("("), ESearchCase::IgnoreCase, ESearchDir::FromStart, FoundPackedVectorCharIndex );
-		if ( FoundPackedVectorCharIndex == INDEX_NONE )
-			break;
 
-		// We want to know which column this char belong to
-		int32 FoundPackedVectorColumnIndex = INDEX_NONE;
-		{
-			// Chop the first line up to the found character
-			FString FirstLineLeft = FirstLine.Left( FoundPackedVectorCharIndex );
-
-			// ReplaceInLine returns the number of occurences of ",", that's what we want! 
-			FoundPackedVectorColumnIndex = FirstLineLeft.ReplaceInline(TEXT(","), TEXT(""));
-		}
+		// Parse the current line to an array
+		TArray<FString> CurrentParsedLine;
+		CurrentLine.ParseIntoArray( CurrentParsedLine, TEXT(",") );
 
-		if ( !TitleRowArray.IsValidIndex( FoundPackedVectorColumnIndex ) )
-		{
+		// Check the parsed line and number of columns match
+		if ( NumberOfColumns != CurrentParsedLine.Num() )
 			UE_LOG( LogHoudiniNiagara, Warning,
-			TEXT( "Error while parsing the CSV File. Couldn't unpack vector found at character %d in the first line!" ),
-			FoundPackedVectorCharIndex + 1 );
-			continue;
-		}
-
-		// We found a packed vector, get its size
-		int32 FoundVectorSize = 0;
-		{
-			// Extract the vector string
-			int32 FoundPackedVectorEndCharIndex = FirstLine.Find( TEXT(")"), ESearchCase::IgnoreCase, ESearchDir::FromStart, FoundPackedVectorCharIndex );
-			FString VectorString = FirstLine.Mid( FoundPackedVectorCharIndex + 1, FoundPackedVectorEndCharIndex - FoundPackedVectorCharIndex - 1 );
+			TEXT("Error while parsing the CSV File. Line %d has %d values instead of the expected %d!"),
+			rowIdx + 1, CurrentParsedLine.Num(), NumberOfColumns );
 
-			// Use ReplaceInLine to count the number of , to get the vector's size!
-			FoundVectorSize = VectorString.ReplaceInline( TEXT(","), TEXT("") ) + 1;
-		}
+		// Store the parsed line
+		ParsedStringArrays[ rowIdx ] = CurrentParsedLine;
+	}
 
-		if ( FoundVectorSize < 2 )
-			continue;
+	// If we have time values, we have to make sure the lines are sorted by time
+	if ( TimeColumnIndex != INDEX_NONE )
+	{
+		// First check if we need to sort the array
+		bool NeedToSort = false;
+		float PreviousTimeValue = 0.0f;
+		for ( int32 rowIdx = 0; rowIdx < ParsedStringArrays.Num(); rowIdx++ )
+		{
+			if ( !ParsedStringArrays[ rowIdx ].IsValidIndex( TimeColumnIndex ) )
+				continue;
 
-		// Increment the number of columns
-		NumberOfColumns += ( FoundVectorSize - 1 );
+			// Get the current time value
+			float CurrentValue = FCString::Atof(*(ParsedStringArrays[ rowIdx ][ TimeColumnIndex ]));
+			if ( rowIdx == 0 )
+			{
+				PreviousTimeValue = CurrentValue;
+				continue;
+			}
 
-		// Expand TitleRowArray
-		if ( ( FoundPackedVectorColumnIndex == PositionColumnIndex ) && ( FoundVectorSize == 3 ) )
-		{
-			// Expand P to Px,Py,Pz
-			TitleRowArray[ PositionColumnIndex ] = TEXT("Px");
-			TitleRowArray.Insert( TEXT( "Py" ), PositionColumnIndex + 1 );
-			TitleRowArray.Insert( TEXT( "Pz" ), PositionColumnIndex + 2 );
-		}
-		else if ( ( FoundPackedVectorColumnIndex == NormalColumnIndex ) && ( FoundVectorSize == 3 ) )
-		{
-			// Expand N to Nx,Ny,Nz
-			TitleRowArray[ NormalColumnIndex ] = TEXT("Nx");
-			TitleRowArray.Insert( TEXT("Ny"), NormalColumnIndex + 1 );
-			TitleRowArray.Insert( TEXT("Nz"), NormalColumnIndex + 2 );
-		}
-		else
-		{
-			// Expand the vector's title from V to V, V1, V2, V3 ...
-			FString FoundPackedVectortTitle = TitleRowArray[ FoundPackedVectorColumnIndex ];
-			for ( int32 n = 1; n < FoundVectorSize; n++ )
+			// Time values arent sorted properly
+			if ( PreviousTimeValue > CurrentValue )
 			{
-				FString CurrentTitle = FoundPackedVectortTitle + FString::FromInt( n );
-				TitleRowArray.Insert( CurrentTitle, FoundPackedVectorColumnIndex + n );
+				NeedToSort = true;
+				break;
 			}
 		}
 
-		// Eventually offset the stored index
-		if ( PositionColumnIndex != INDEX_NONE && ( PositionColumnIndex > FoundPackedVectorColumnIndex) )
-			PositionColumnIndex += FoundVectorSize - 1;
-
-		if ( NormalColumnIndex != INDEX_NONE && ( NormalColumnIndex > FoundPackedVectorColumnIndex) )
-			NormalColumnIndex += FoundVectorSize - 1;
-
-		if ( TimeColumnIndex != INDEX_NONE && ( TimeColumnIndex > FoundPackedVectorColumnIndex) )
-			TimeColumnIndex += FoundVectorSize - 1;
-
-		if ( IDColumnIndex != INDEX_NONE && ( IDColumnIndex > FoundPackedVectorColumnIndex ) )
-			IDColumnIndex += FoundVectorSize - 1;
-
-		if ( AliveColumnIndex != INDEX_NONE && ( AliveColumnIndex > FoundPackedVectorColumnIndex ) )
-			AliveColumnIndex += FoundVectorSize - 1;
-
-		if ( LifeColumnIndex != INDEX_NONE && ( LifeColumnIndex > FoundPackedVectorColumnIndex ) )
-			LifeColumnIndex += FoundVectorSize - 1;
-
-		HasPackedVectors = true;
-		FoundPackedVectorCharIndex++;
-    }
-
-    // For sanity, Check that the number of columns matches the title row and the first line
-    {
-		// Check the title row
-		if ( NumberOfColumns != TitleRowArray.Num() )
-			UE_LOG( LogHoudiniNiagara, Error,
-			TEXT( "Error while parsing the CSV File. Found %d columns but the Title string has %d values! Some values will have an offset!" ),
-			NumberOfColumns, TitleRowArray.Num() );
-
-		// Use ReplaceInLine to count the number of columns in the first line and make sure it's correct
-		int32 FirstLineColumns = FirstLine.ReplaceInline( TEXT(","), TEXT("") ) + 1;
-		if ( NumberOfColumns != FirstLineColumns )
-			UE_LOG( LogHoudiniNiagara, Error,
-			TEXT("Error while parsing the CSV File. Found %d columns but found %d values in the first line! Some values will have an offset!" ),
-			NumberOfColumns, FirstLineColumns );
-    }
-	*/
-    
-    // Remove the title row now that it's been processed
-    RawStringArray.RemoveAt( 0 );
+		if ( NeedToSort )
+		{
+			// We need to sort the CSV lines by their time values
+			ParsedStringArrays.Sort<FHoudiniCSVSortPredicate>( FHoudiniCSVSortPredicate( TimeColumnIndex, IDColumnIndex ) );
+		}
+	}
 
     // Initialize our different buffers
     FloatCSVData.Empty();
@@ -290,30 +223,17 @@ bool UHoudiniCSV::UpdateFromStringArray( TArray<FString>& RawStringArray )
     StringCSVData.Empty();
     StringCSVData.SetNumZeroed( NumberOfLines * NumberOfColumns );
 
-	TimeValues.Empty();
+	// Due to the way that some of the DI functions work,
+	// we expect that the particle IDs start at zero, and increment as the particles are spawned
+	// Make sure this is the case by converting the particle IDs as we read them
+	int32 NextParticleID = 0;
+	TMap<float, int32> HoudiniIDToNiagaraIDMap;
 
     // Extract all the values from the table to the float & string buffers
     TArray<FString> CurrentParsedLine;
-    for ( int rowIdx = 0; rowIdx < NumberOfLines; rowIdx++ )
+    for ( int rowIdx = 0; rowIdx < ParsedStringArrays.Num(); rowIdx++ )
     {
-		FString CurrentLine = RawStringArray[ rowIdx ];
-		if ( HasPackedVectors )
-		{
-			// Clean up the packing characters: ()" from the line so it can be parsed properly
-			CurrentLine.ReplaceInline( TEXT("("), TEXT("") );
-			CurrentLine.ReplaceInline( TEXT(")"), TEXT("") );
-			CurrentLine.ReplaceInline( TEXT("\""), TEXT("") );
-		}
-
-		// Parse the current line to an array
-		CurrentParsedLine.Empty();
-		CurrentLine.ParseIntoArray( CurrentParsedLine, TEXT(",") );
-
-		// Check the parsed line and number of columns match
-		if ( NumberOfColumns != CurrentParsedLine.Num() )
-			UE_LOG( LogHoudiniNiagara, Warning,
-			TEXT("Error while parsing the CSV File. Line %d has %d values instead of the expected %d!"),
-			rowIdx + 1, CurrentParsedLine.Num(), NumberOfColumns );
+		CurrentParsedLine = ParsedStringArrays[ rowIdx ];
 
 		// Store the CSV Data in the buffers
 		// The data is stored transposed in those buffers
@@ -323,7 +243,9 @@ bool UHoudiniCSV::UpdateFromStringArray( TArray<FString>& RawStringArray )
 			// Get the string value for the current column
 			FString CurrentVal = TEXT("0");
 			if ( CurrentParsedLine.IsValidIndex( colIdx ) )
+			{
 				CurrentVal = CurrentParsedLine[ colIdx ];
+			}
 			else
 			{
 				UE_LOG( LogHoudiniNiagara, Warning,
@@ -331,29 +253,37 @@ bool UHoudiniCSV::UpdateFromStringArray( TArray<FString>& RawStringArray )
 				rowIdx + 1, colIdx + 1 );
 			}
 
-			// Convert the string value to a float in the buffer
+			// Convert the string value to a float
 			float FloatValue = FCString::Atof( *CurrentVal );
+
+			// Handle particles IDs here
+			if ( colIdx == IDColumnIndex )
+			{
+				// The particle ID may need to be replaced
+				if ( !HoudiniIDToNiagaraIDMap.Contains( FloatValue ) )
+					HoudiniIDToNiagaraIDMap.Add( FloatValue, NextParticleID++ );
+
+				CurrentID = HoudiniIDToNiagaraIDMap[ FloatValue ];
+				FloatValue = (float)CurrentID;
+			}
+
+			// Store the Value in the buffer
 			FloatCSVData[ rowIdx + ( colIdx * NumberOfLines ) ] = FloatValue;
 
 			// Keep the original string value in a buffer too
 			StringCSVData[ rowIdx + ( colIdx * NumberOfLines ) ] = CurrentVal;
-						
-			if ( colIdx == IDColumnIndex )
-				CurrentID = (int32)FloatValue;
-
-			if ( colIdx == TimeColumnIndex )
-				TimeValues.AddUnique( FloatValue );
 		}
 
 		// Look at the particle ID, the max ID will be our number of particles
 		if ( NumberOfParticles <= CurrentID )
 			NumberOfParticles = CurrentID + 1;
     }
-
+	
+	NumberOfParticles = HoudiniIDToNiagaraIDMap.Num();
 	if ( NumberOfParticles <= 0 )
 		NumberOfParticles = NumberOfLines;
 
-	// Look for particles specific attributes
+	// Look for particle specific attributes
 	SpawnTimes.Empty();
 	SpawnTimes.Init( -1.0f, NumberOfParticles );
 
@@ -368,7 +298,7 @@ bool UHoudiniCSV::UpdateFromStringArray( TArray<FString>& RawStringArray )
 		if ( IDColumnIndex != INDEX_NONE )
 			CurrentID = (int32)FloatCSVData[ rowIdx + ( IDColumnIndex * NumberOfLines ) ];
 
-		// Get the current time
+		// Get the time value for the current line
 		float CurrentTime = 0.0f;
 		if ( TimeColumnIndex != INDEX_NONE )
 			CurrentTime = FloatCSVData[ rowIdx + ( TimeColumnIndex * NumberOfLines ) ];
@@ -424,10 +354,10 @@ bool UHoudiniCSV::ParseCSVTitleLine( const FString& TitleLine, const FString& Fi
     // Look for the position, normal and time attributes indexes
     for ( int32 n = 0; n < TitleRowArray.Num(); n++ )
     {
-		FString CurrentTitle = TitleRowArray[ n ];
-		CurrentTitle.ToLowerInline();
-		CurrentTitle.ReplaceInline( TEXT(" "), TEXT("") );
+		// Remove spaces from the title row
+		TitleRowArray[ n ].ReplaceInline( TEXT(" "), TEXT("") );
 
+		FString CurrentTitle = TitleRowArray[ n ];
 		if ( CurrentTitle.Equals( TEXT("P"), ESearchCase::IgnoreCase )
 			|| CurrentTitle.Equals( TEXT("Px"), ESearchCase::IgnoreCase )
 			|| CurrentTitle.Equals( TEXT("X"), ESearchCase::IgnoreCase )
@@ -466,10 +396,10 @@ bool UHoudiniCSV::ParseCSVTitleLine( const FString& TitleLine, const FString& Fi
 		}
     }
 
-	 // Read the first line of the CSV file, and look for packed vectors value (X,Y,Z)
+	// Read the first line of the CSV file, and look for packed vectors value (X,Y,Z)
     // We'll have to expand them in the title row to match the parsed data
-    int32 FoundPackedVectorCharIndex = 0;
-    HasPackedVectors = false;
+	HasPackedVectors = false;
+	int32 FoundPackedVectorCharIndex = 0;    
     while ( FoundPackedVectorCharIndex != INDEX_NONE )
     {
 		// Try to find ( in the line
@@ -1000,8 +930,6 @@ bool UHoudiniCSV::GetCSVStringValue( const int32& lineIndex, const FString& Colu
     return GetCSVStringValue( lineIndex, ColIndex, value );
 }
 
-// 
-
 #if WITH_EDITORONLY_DATA
 void UHoudiniCSV::PostInitProperties()
 {
diff --git a/Engine/Plugins/FX/HoudiniNiagara/Source/HoudiniNiagara/Public/HoudiniCSV.h b/Engine/Plugins/FX/HoudiniNiagara/Source/HoudiniNiagara/Public/HoudiniCSV.h
--- a/Engine/Plugins/FX/HoudiniNiagara/Source/HoudiniNiagara/Public/HoudiniCSV.h
+++ b/Engine/Plugins/FX/HoudiniNiagara/Source/HoudiniNiagara/Public/HoudiniCSV.h
@@ -138,6 +138,8 @@ class HOUDININIAGARA_API UHoudiniCSV : public UObject
 
     private:
 
+	static bool SortPredicate(const TArray<FString>& A, const TArray<FString>& B);
+
 	// Array containing the Raw String data
 	UPROPERTY()
 	TArray<FString> StringCSVData;
@@ -154,9 +156,6 @@ class HOUDININIAGARA_API UHoudiniCSV : public UObject
 	UPROPERTY()
 	TArray<float> LifeValues;
 
-	// Array containing the different time values for each particles in the file
-	TArray<float> TimeValues;
-
 	// Index of the Position values in the buffer
 	UPROPERTY()	    
 	int32 PositionColumnIndex;
