Niagara - Houdini:

Houdini CSV Asset:
When importing the CSV file, the importer creates a list of the row indexes updating each particle.
This greatly improves performance when accessing data in large files with a lot of particles updating over time.

diff --git a/Engine/Plugins/FX/HoudiniNiagara/Source/HoudiniNiagara/Private/HoudiniCSV.cpp b/Engine/Plugins/FX/HoudiniNiagara/Source/HoudiniNiagara/Private/HoudiniCSV.cpp
--- a/Engine/Plugins/FX/HoudiniNiagara/Source/HoudiniNiagara/Private/HoudiniCSV.cpp
+++ b/Engine/Plugins/FX/HoudiniNiagara/Source/HoudiniNiagara/Private/HoudiniCSV.cpp
@@ -229,6 +229,13 @@ bool UHoudiniCSV::UpdateFromStringArray( TArray<FString>& RawStringArray )
 	int32 NextParticleID = 0;
 	TMap<float, int32> HoudiniIDToNiagaraIDMap;
 
+	// We also keep track of the row indexes for each time values
+	//float lastTimeValue = 0.0;
+	//TimeValuesIndexes.Empty();
+
+	// And the row indexes for each particle
+	ParticleValueIndexes.Empty();
+
     // Extract all the values from the table to the float & string buffers
     TArray<FString> CurrentParsedLine;
     for ( int rowIdx = 0; rowIdx < ParsedStringArrays.Num(); rowIdx++ )
@@ -261,11 +268,33 @@ bool UHoudiniCSV::UpdateFromStringArray( TArray<FString>& RawStringArray )
 			{
 				// The particle ID may need to be replaced
 				if ( !HoudiniIDToNiagaraIDMap.Contains( FloatValue ) )
+				{
+					// We found a new particle, add it to the ID map
 					HoudiniIDToNiagaraIDMap.Add( FloatValue, NextParticleID++ );
 
+					// Add a new array for that particle's indexes
+					ParticleValueIndexes.Add( FParticleIndexes() );
+				}
+
+				// Get the Niagara ID from this Houdini ID
 				CurrentID = HoudiniIDToNiagaraIDMap[ FloatValue ];
 				FloatValue = (float)CurrentID;
+
+				// Add the current row to this particle's row index list
+				ParticleValueIndexes[ CurrentID ].RowIndexes.Add( rowIdx );
+			}
+
+			/*
+			if ( colIdx == TimeColumnIndex )
+			{
+				// Keep track of new time values row indexes
+				if ( ( FloatValue != lastTimeValue ) || ( rowIdx == 0 ) )
+				{
+					TimeValuesIndexes.Add( FloatValue, rowIdx );
+					lastTimeValue = FloatValue;
+				}
 			}
+			*/
 
 			// Store the Value in the buffer
 			FloatCSVData[ rowIdx + ( colIdx * NumberOfLines ) ] = FloatValue;
@@ -825,17 +854,25 @@ bool UHoudiniCSV::GetParticleLineIndexAtTime(const int32& ParticleID, const floa
 	if ( ParticleLife > 0.0f && ( ParticleSpawnTime + ParticleLife < desiredTime ) )	
 		return false;
 
-	for ( int32 n = 0; n < NumberOfLines; n++ )
-	{
-		// Get the ID
-		float CurrentLineID = -1;
-		if ( !GetCSVFloatValue( n, IDColumnIndex, CurrentLineID ) )
-			CurrentLineID = n;
+	// We don't have Id information
+	// return ParticleId for the rowIndexes ??
+	if ( IDColumnIndex == INDEX_NONE )
+		return false;
 
-		// Not the particle we're looking for
-		if ( CurrentLineID != ParticleID )
-			continue;
+	// We don't have time information
+	if ( TimeColumnIndex == INDEX_NONE )
+		return false;
+
+	// Get the row indexes for this particle
+	TArray<int32>* RowIndexes = nullptr;
+	if ( ParticleValueIndexes.IsValidIndex( ParticleID ) )
+		RowIndexes = &( ParticleValueIndexes[ ParticleID ].RowIndexes );
 
+	if ( !RowIndexes )
+		return false;
+
+	for ( auto n : *RowIndexes )
+	{
 		// Get the time
 		float currentTime = -1.0f;
 		if ( GetCSVTimeValue( n, currentTime) )
diff --git a/Engine/Plugins/FX/HoudiniNiagara/Source/HoudiniNiagara/Public/HoudiniCSV.h b/Engine/Plugins/FX/HoudiniNiagara/Source/HoudiniNiagara/Public/HoudiniCSV.h
--- a/Engine/Plugins/FX/HoudiniNiagara/Source/HoudiniNiagara/Public/HoudiniCSV.h
+++ b/Engine/Plugins/FX/HoudiniNiagara/Source/HoudiniNiagara/Public/HoudiniCSV.h
@@ -29,6 +29,17 @@
 
 DECLARE_LOG_CATEGORY_EXTERN( LogHoudiniNiagara, All, All );
 
+
+USTRUCT()
+struct FParticleIndexes
+{
+	GENERATED_BODY()
+
+	// Simple structure for storing all the row indexes for a particle
+	UPROPERTY()
+	TArray<int32> RowIndexes;
+};
+
 UCLASS()
 class HOUDININIAGARA_API UHoudiniCSV : public UObject
 {
@@ -179,4 +190,20 @@ class HOUDININIAGARA_API UHoudiniCSV : public UObject
 	// Index of the life values in the buffer
 	UPROPERTY()
 	int32 LifeColumnIndex;
+
+	// Index of the color values in the buffer
+	//UPROPERTY()
+	//int32 ColorColumnIndex;
+
+	// Index of the velocity values in the buffer
+	//UPROPERTY()
+	//int32 VelocityColumnIndex;
+
+	// Row indexes for new time values
+	//UPROPERTY()
+	//TMap<float, int32> TimeValuesIndexes;
+
+	// Row indexes for each particle
+	UPROPERTY()
+	TArray< FParticleIndexes > ParticleValueIndexes;
 };
\ No newline at end of file
