Niagara - Houdini:

Houdini Niagara Data Interface:
Added helper functions for accessing Color and Velocity values in the CSV file.

Houdini CSV Asset:
The importer now looks for the Color (Cd, color), Alpha (A, Alpha) and velocity (V) attributes.

diff --git a/Engine/Plugins/FX/HoudiniNiagara/Source/HoudiniNiagara/Classes/NiagaraDataInterfaceHoudiniCSV.h b/Engine/Plugins/FX/HoudiniNiagara/Source/HoudiniNiagara/Classes/NiagaraDataInterfaceHoudiniCSV.h
--- a/Engine/Plugins/FX/HoudiniNiagara/Source/HoudiniNiagara/Classes/NiagaraDataInterfaceHoudiniCSV.h
+++ b/Engine/Plugins/FX/HoudiniNiagara/Source/HoudiniNiagara/Classes/NiagaraDataInterfaceHoudiniCSV.h
@@ -61,7 +61,7 @@ public:
 	//----------------------------------------------------------------------------
 	// EXPOSED FUNCTIONS
 
-	// Returns the float value at a given point in the CSV file
+	// Returns the float value at a given row and column in the CSV file
 	template<typename RowParamType, typename ColParamType>
 	void GetCSVFloatValue(FVectorVMContext& Context);
 
@@ -70,23 +70,31 @@ public:
 	void GetCSVFloatValueByString(FVectorVMContext& Context);
 	*/
 
-	// Returns a Vector3 value for a given point in the csv file
+	// Returns a Vector3 value for a given row in the csv file
 	template<typename RowParamType, typename ColParamType>
 	void GetCSVVectorValue(FVectorVMContext& Context);
 
-	// Returns the positions for a given point in the CSV file
+	// Returns the positions for a given row in the CSV file
 	template<typename RowParamType>
 	void GetCSVPosition(FVectorVMContext& Context);
 
-	// Returns the normals for a given point in the CSV file
+	// Returns the normals for a given row in the CSV file
 	template<typename RowParamType>
 	void GetCSVNormal(FVectorVMContext& Context);
 
-	// Returns the time for a given point in the CSV file
+	// Returns the time for a given row in the CSV file
 	template<typename RowParamType>
 	void GetCSVTime(FVectorVMContext& Context);
 
-	// Returns the position and time for a given point in the CSV file
+	// Returns the velocity for a given row in the CSV file
+	template<typename RowParamType>
+	void GetCSVVelocity(FVectorVMContext& Context);
+
+	// Returns the color for a given row in the CSV file
+	template<typename RowParamType>
+	void GetCSVColor(FVectorVMContext& Context);
+
+	// Returns the position and time for a given row in the CSV file
 	template<typename RowParamType>
 	void GetCSVPositionAndTime(FVectorVMContext& Context);
 
diff --git a/Engine/Plugins/FX/HoudiniNiagara/Source/HoudiniNiagara/Private/HoudiniCSV.cpp b/Engine/Plugins/FX/HoudiniNiagara/Source/HoudiniNiagara/Private/HoudiniCSV.cpp
--- a/Engine/Plugins/FX/HoudiniNiagara/Source/HoudiniNiagara/Private/HoudiniCSV.cpp
+++ b/Engine/Plugins/FX/HoudiniNiagara/Source/HoudiniNiagara/Private/HoudiniCSV.cpp
@@ -85,6 +85,9 @@ UHoudiniCSV::UHoudiniCSV( const FObjectInitializer& ObjectInitializer )
 	IDColumnIndex = INDEX_NONE;
 	AliveColumnIndex = INDEX_NONE;
 	LifeColumnIndex = INDEX_NONE;
+	ColorColumnIndex = INDEX_NONE;
+	AlphaColumnIndex = INDEX_NONE;
+	VelocityColumnIndex = INDEX_NONE;
 }
 
 void UHoudiniCSV::SetFileName( const FString& TheFileName )
@@ -126,6 +129,9 @@ bool UHoudiniCSV::UpdateFromStringArray( TArray<FString>& RawStringArray )
 	IDColumnIndex = INDEX_NONE;
 	AliveColumnIndex = INDEX_NONE;
 	LifeColumnIndex = INDEX_NONE;
+	ColorColumnIndex = INDEX_NONE;
+	AlphaColumnIndex = INDEX_NONE;
+	VelocityColumnIndex = INDEX_NONE;
 
     if ( RawStringArray.Num() <= 0 )
     {
@@ -415,7 +421,7 @@ bool UHoudiniCSV::ParseCSVTitleLine( const FString& TitleLine, const FString& Fi
 		}
 		else if ( CurrentTitle.Equals( TEXT("alive"), ESearchCase::IgnoreCase ) )
 		{
-			if (AliveColumnIndex == INDEX_NONE)
+			if ( AliveColumnIndex == INDEX_NONE )
 				AliveColumnIndex = n;
 		}
 		else if ( CurrentTitle.Equals( TEXT("life"), ESearchCase::IgnoreCase ) )
@@ -423,6 +429,24 @@ bool UHoudiniCSV::ParseCSVTitleLine( const FString& TitleLine, const FString& Fi
 			if ( LifeColumnIndex == INDEX_NONE)
 				LifeColumnIndex = n;
 		}
+		else if ( ( CurrentTitle.Equals( TEXT("Cd"), ESearchCase::IgnoreCase ) )
+			|| ( CurrentTitle.Equals(TEXT("color"), ESearchCase::IgnoreCase ) ) )
+		{
+			if ( ColorColumnIndex == INDEX_NONE )
+				ColorColumnIndex = n;
+		}
+		else if ( ( CurrentTitle.Equals( TEXT("alpha"), ESearchCase::IgnoreCase ) )
+			|| ( CurrentTitle.Equals(TEXT("A"), ESearchCase::IgnoreCase ) ) )
+		{
+			if ( AlphaColumnIndex == INDEX_NONE )
+				AlphaColumnIndex = n;
+		}
+		else if ( ( CurrentTitle.Equals( TEXT("v"), ESearchCase::IgnoreCase ) )
+			|| ( CurrentTitle.Equals(TEXT("Vx"), ESearchCase::IgnoreCase ) ) )
+		{
+			if ( VelocityColumnIndex == INDEX_NONE )
+				VelocityColumnIndex = n;
+		}
     }
 
 	// Read the first line of the CSV file, and look for packed vectors value (X,Y,Z)
@@ -486,6 +510,28 @@ bool UHoudiniCSV::ParseCSVTitleLine( const FString& TitleLine, const FString& Fi
 			TitleRowArray.Insert( TEXT("Ny"), NormalColumnIndex + 1 );
 			TitleRowArray.Insert( TEXT("Nz"), NormalColumnIndex + 2 );
 		}
+		else if ( ( FoundPackedVectorColumnIndex == VelocityColumnIndex ) && ( FoundVectorSize == 3 ) )
+		{
+			// Expand V to Vx,Vy,Vz
+			TitleRowArray[ VelocityColumnIndex ] = TEXT("Vx");
+			TitleRowArray.Insert( TEXT("Vy"), VelocityColumnIndex + 1 );
+			TitleRowArray.Insert( TEXT("Vz"), VelocityColumnIndex + 2 );
+		}
+		else if ( ( FoundPackedVectorColumnIndex == ColorColumnIndex ) && ( ( FoundVectorSize == 3 ) || ( FoundVectorSize == 4 ) ) )
+		{
+			// Expand Cd to R, G, B 
+			TitleRowArray[ ColorColumnIndex ] = TEXT("R");
+			TitleRowArray.Insert( TEXT("G"), ColorColumnIndex + 1 );
+			TitleRowArray.Insert( TEXT("B"), ColorColumnIndex + 2 );
+
+			if ( FoundVectorSize == 4 )
+			{
+				// Insert A if we had RGBA
+				TitleRowArray.Insert( TEXT("A"), ColorColumnIndex + 3 );
+				if ( AlphaColumnIndex == INDEX_NONE )
+					AlphaColumnIndex = ColorColumnIndex + 3;
+			}
+		}
 		else
 		{
 			// Expand the vector's title from V to V, V1, V2, V3 ...
@@ -516,6 +562,15 @@ bool UHoudiniCSV::ParseCSVTitleLine( const FString& TitleLine, const FString& Fi
 		if ( LifeColumnIndex != INDEX_NONE && ( LifeColumnIndex > FoundPackedVectorColumnIndex ) )
 			LifeColumnIndex += FoundVectorSize - 1;
 
+		if ( ColorColumnIndex != INDEX_NONE && ( ColorColumnIndex > FoundPackedVectorColumnIndex ) )
+			ColorColumnIndex += FoundVectorSize - 1;
+
+		if ( AlphaColumnIndex != INDEX_NONE && ( AlphaColumnIndex > FoundPackedVectorColumnIndex ) )
+			AlphaColumnIndex += FoundVectorSize - 1;
+
+		if ( VelocityColumnIndex != INDEX_NONE && ( VelocityColumnIndex > FoundPackedVectorColumnIndex ) )
+			VelocityColumnIndex += FoundVectorSize - 1;
+
 		HasPackedVectors = true;
 		FoundPackedVectorCharIndex++;
     }
@@ -642,6 +697,39 @@ bool UHoudiniCSV::GetCSVTimeValue( const int32& lineIndex, float& value )
     return true;
 }
 
+// Returns a Color for a given point in the CSV file
+bool UHoudiniCSV::GetCSVColorValue( const int32& lineIndex, FLinearColor& value )
+{
+	FVector V = FVector::OneVector;
+	if ( !GetCSVVectorValue( lineIndex, ColorColumnIndex, V, false, false ) )
+		return false;
+
+	FLinearColor C = FLinearColor::White;
+	C.R = V.X;
+	C.G = V.Y;
+	C.B = V.Z;
+
+	float alpha = 1.0f;
+	if ( GetCSVFloatValue( lineIndex, AlphaColumnIndex, alpha ) )
+		C.A = alpha;	
+
+	value = C;
+
+	return true;
+}
+
+// Returns a Velocity Vector3 for a given point in the CSV file
+bool UHoudiniCSV::GetCSVVelocityValue( const int32& lineIndex, FVector& value )
+{
+	FVector V = FVector::ZeroVector;
+	if ( !GetCSVVectorValue( lineIndex, VelocityColumnIndex, V, true, false ) )
+		return false;
+
+	value = V;
+
+	return true;
+}
+
 // Returns the number of points found in the CSV file
 int32 UHoudiniCSV::GetNumberOfParticlesInCSV()
 {
diff --git a/Engine/Plugins/FX/HoudiniNiagara/Source/HoudiniNiagara/Private/NiagaraDataInterfaceHoudiniCSV.cpp b/Engine/Plugins/FX/HoudiniNiagara/Source/HoudiniNiagara/Private/NiagaraDataInterfaceHoudiniCSV.cpp
--- a/Engine/Plugins/FX/HoudiniNiagara/Source/HoudiniNiagara/Private/NiagaraDataInterfaceHoudiniCSV.cpp
+++ b/Engine/Plugins/FX/HoudiniNiagara/Source/HoudiniNiagara/Private/NiagaraDataInterfaceHoudiniCSV.cpp
@@ -225,6 +225,38 @@ void UNiagaraDataInterfaceHoudiniCSV::GetFunctions(TArray<FNiagaraFunctionSignat
 		OutFunctions.Add(Sig);
     }
 
+	{
+		// GetCSVVelocity
+		FNiagaraFunctionSignature Sig;
+		Sig.Name = TEXT("GetCSVVelocity");
+		Sig.bMemberFunction = true;
+		Sig.bRequiresContext = false;
+		Sig.Inputs.Add(FNiagaraVariable(FNiagaraTypeDefinition(GetClass()), TEXT("CSV")));			// CSV in
+		Sig.Inputs.Add(FNiagaraVariable(FNiagaraTypeDefinition::GetIntDef(), TEXT("Row")));			// Row Index In
+		Sig.Outputs.Add(FNiagaraVariable(FNiagaraTypeDefinition::GetVec3Def(), TEXT("Velocity")));	// Vector3 Out
+
+		Sig.SetDescription(LOCTEXT("DataInterfaceHoudini_GetCSVVelocity",
+			"Helper function returning the velocity value for a given Row in the CSV file.\nThe returned velocity vector is converted from Houdini's coordinate system to Unreal's."));
+
+		OutFunctions.Add(Sig);
+	}
+
+	{
+		// GetCSVColor
+		FNiagaraFunctionSignature Sig;
+		Sig.Name = TEXT("GetCSVColor");
+		Sig.bMemberFunction = true;
+		Sig.bRequiresContext = false;
+		Sig.Inputs.Add(FNiagaraVariable(FNiagaraTypeDefinition(GetClass()), TEXT("CSV")));			// CSV in
+		Sig.Inputs.Add(FNiagaraVariable(FNiagaraTypeDefinition::GetIntDef(), TEXT("Row")));			// Row Index In
+		Sig.Outputs.Add(FNiagaraVariable(FNiagaraTypeDefinition::GetColorDef(), TEXT("Color")));	// Color Out
+
+		Sig.SetDescription(LOCTEXT("DataInterfaceHoudini_GetCSVColor",
+			"Helper function returning the color value for a given Row in the CSV file."));
+
+		OutFunctions.Add(Sig);
+	}
+
     {
 		// GetNumberOfParticlesInCSV
 		FNiagaraFunctionSignature Sig;
@@ -483,6 +515,8 @@ DEFINE_NDI_RAW_FUNC_BINDER(UNiagaraDataInterfaceHoudiniCSV, GetCSVVectorValue);
 DEFINE_NDI_RAW_FUNC_BINDER(UNiagaraDataInterfaceHoudiniCSV, GetCSVPosition);
 DEFINE_NDI_RAW_FUNC_BINDER(UNiagaraDataInterfaceHoudiniCSV, GetCSVNormal);
 DEFINE_NDI_RAW_FUNC_BINDER(UNiagaraDataInterfaceHoudiniCSV, GetCSVTime);
+DEFINE_NDI_RAW_FUNC_BINDER(UNiagaraDataInterfaceHoudiniCSV, GetCSVColor);
+DEFINE_NDI_RAW_FUNC_BINDER(UNiagaraDataInterfaceHoudiniCSV, GetCSVVelocity);
 DEFINE_NDI_RAW_FUNC_BINDER(UNiagaraDataInterfaceHoudiniCSV, GetCSVPositionAndTime);
 DEFINE_NDI_RAW_FUNC_BINDER(UNiagaraDataInterfaceHoudiniCSV, GetLastRowIndexAtTime);
 DEFINE_NDI_RAW_FUNC_BINDER(UNiagaraDataInterfaceHoudiniCSV, GetParticleIndexesToSpawnAtTime);
@@ -517,6 +551,14 @@ void UNiagaraDataInterfaceHoudiniCSV::GetVMExternalFunction(const FVMExternalFun
     {
 		TNDIParamBinder<0, int32, NDI_RAW_FUNC_BINDER(UNiagaraDataInterfaceHoudiniCSV, GetCSVTime)>::Bind(this, BindingInfo, InstanceData, OutFunc);
     }
+	else if (BindingInfo.Name == TEXT("GetCSVVelocity") && BindingInfo.GetNumInputs() == 1 && BindingInfo.GetNumOutputs() == 3)
+	{
+		TNDIParamBinder<0, int32, NDI_RAW_FUNC_BINDER(UNiagaraDataInterfaceHoudiniCSV, GetCSVVelocity)>::Bind(this, BindingInfo, InstanceData, OutFunc);
+	}
+	else if (BindingInfo.Name == TEXT("GetCSVColor") && BindingInfo.GetNumInputs() == 1 && BindingInfo.GetNumOutputs() == 4)
+	{
+		TNDIParamBinder<0, int32, NDI_RAW_FUNC_BINDER(UNiagaraDataInterfaceHoudiniCSV, GetCSVColor)>::Bind(this, BindingInfo, InstanceData, OutFunc);
+	}
     else if (BindingInfo.Name == TEXT("GetCSVPositionAndTime") && BindingInfo.GetNumInputs() == 1 && BindingInfo.GetNumOutputs() == 4)
     {
 		TNDIParamBinder<0, int32, NDI_RAW_FUNC_BINDER(UNiagaraDataInterfaceHoudiniCSV, GetCSVPositionAndTime)>::Bind(this, BindingInfo, InstanceData, OutFunc);
@@ -722,6 +764,61 @@ void UNiagaraDataInterfaceHoudiniCSV::GetCSVTime(FVectorVMContext& Context)
     }
 }
 
+template<typename RowParamType>
+void UNiagaraDataInterfaceHoudiniCSV::GetCSVVelocity(FVectorVMContext& Context)
+{
+	RowParamType RowParam(Context);
+	FRegisterHandler<float> OutSampleX(Context);
+	FRegisterHandler<float> OutSampleY(Context);
+	FRegisterHandler<float> OutSampleZ(Context);
+
+	for (int32 i = 0; i < Context.NumInstances; ++i)
+	{
+		int32 row = RowParam.Get();
+
+		FVector V = FVector::ZeroVector;
+		if ( CSVFile )
+			CSVFile->GetCSVVelocityValue( row, V );
+
+		*OutSampleX.GetDest() = V.X;
+		*OutSampleY.GetDest() = V.Y;
+		*OutSampleZ.GetDest() = V.Z;
+		RowParam.Advance();
+		OutSampleX.Advance();
+		OutSampleY.Advance();
+		OutSampleZ.Advance();
+	}
+}
+
+template<typename RowParamType>
+void UNiagaraDataInterfaceHoudiniCSV::GetCSVColor(FVectorVMContext& Context)
+{
+	RowParamType RowParam(Context);
+	FRegisterHandler<float> OutSampleR(Context);
+	FRegisterHandler<float> OutSampleG(Context);
+	FRegisterHandler<float> OutSampleB(Context);
+	FRegisterHandler<float> OutSampleA(Context);
+
+	for (int32 i = 0; i < Context.NumInstances; ++i)
+	{
+		int32 row = RowParam.Get();
+
+		FLinearColor C = FLinearColor::White;
+		if ( CSVFile )
+			CSVFile->GetCSVColorValue( row, C );
+
+		*OutSampleR.GetDest() = C.R;
+		*OutSampleG.GetDest() = C.G;
+		*OutSampleB.GetDest() = C.B;
+		*OutSampleA.GetDest() = C.A;
+		RowParam.Advance();
+		OutSampleR.Advance();
+		OutSampleG.Advance();
+		OutSampleB.Advance();
+		OutSampleA.Advance();
+	}
+}
+
 // Returns the last index of the particles that should be spawned at time t
 template<typename TimeParamType>
 void UNiagaraDataInterfaceHoudiniCSV::GetLastRowIndexAtTime(FVectorVMContext& Context)
diff --git a/Engine/Plugins/FX/HoudiniNiagara/Source/HoudiniNiagara/Public/HoudiniCSV.h b/Engine/Plugins/FX/HoudiniNiagara/Source/HoudiniNiagara/Public/HoudiniCSV.h
--- a/Engine/Plugins/FX/HoudiniNiagara/Source/HoudiniNiagara/Public/HoudiniCSV.h
+++ b/Engine/Plugins/FX/HoudiniNiagara/Source/HoudiniNiagara/Public/HoudiniCSV.h
@@ -85,6 +85,10 @@ class HOUDININIAGARA_API UHoudiniCSV : public UObject
 	bool GetCSVPositionValue( const int32& lineIndex, FVector& value );
 	// Returns a Normal Vector3 for a given point in the CSV file (converted to unreal's coordinate system)
 	bool GetCSVNormalValue( const int32& lineIndex, FVector& value );
+	// Returns a Color for a given point in the CSV file
+	bool GetCSVColorValue( const int32& lineIndex, FLinearColor& value );
+	// Returns a Velocity Vector3 for a given point in the CSV file
+	bool GetCSVVelocityValue(const int32& lineIndex, FVector& value );
 
 	// Get the last row index for a given time value (the row with a time smaller or equal to desiredTime)
 	// If the CSV file doesn't have time informations, returns false and set the LastRowIndex to the last line in the file
@@ -192,12 +196,16 @@ class HOUDININIAGARA_API UHoudiniCSV : public UObject
 	int32 LifeColumnIndex;
 
 	// Index of the color values in the buffer
-	//UPROPERTY()
-	//int32 ColorColumnIndex;
+	UPROPERTY()
+	int32 ColorColumnIndex;
+
+	// Index of the color values in the buffer
+	UPROPERTY()
+	int32 AlphaColumnIndex;
 
 	// Index of the velocity values in the buffer
-	//UPROPERTY()
-	//int32 VelocityColumnIndex;
+	UPROPERTY()
+	int32 VelocityColumnIndex;
 
 	// Row indexes for new time values
 	//UPROPERTY()
