Houdini - Niagara:
Big renaming pass on the  Houdini CSV Assets  and Data Interface to follow naming conventions:
Replaced the GetCSVXXX functions by GetXXXX (GetCSVPosition is now GetPosition)
Always use "row" instead of "line", "Point" instead of "Particle", "PointID" instead of "N" or "ID" etc.

Houdini Data Interface:
- Added the GetVectorValueEx and GetPointVectorValueAtTimeEx functions that allow the user to decide how the vector conversion from houdini to unreal's coordinate system is handled.
- Replaced the GetParticleLifeAtTime function by GetPointLife, that returns the life of a particle at spawn time.
- Added the GetPointType function returning the type of a given point.

Houdini CSV Asset:
- Added the editable SourceTitleRow UProperty.
Editing this will trigger a reimport of the source CSV file and might be used to fix/modify column titles in the file.
- Added support for "type" attributes.
- Removed the unused StringValues buffer and GetCSVStringValues() functions.
- Added assetTags so the Houdini CSV asset thumbnails show more infos on the CSV data.
- Added the "FindSourceCSV" asset action to browse to the source CSV file.

diff --git a/Engine/Plugins/FX/HoudiniNiagara/HoudiniNiagara.uplugin b/Engine/Plugins/FX/HoudiniNiagara/HoudiniNiagara.uplugin
--- a/Engine/Plugins/FX/HoudiniNiagara/HoudiniNiagara.uplugin
+++ b/Engine/Plugins/FX/HoudiniNiagara/HoudiniNiagara.uplugin
@@ -10,30 +10,31 @@
 	"DocsURL": "",
 	"MarketplaceURL": "",
 	"SupportURL": "",
+	"EnabledByDefault": false,
 	"CanContainContent": true,
 	"IsBetaVersion": true,
 	"Installed": false,
 	"Modules":
-    [
+	[
 		{
 			"Name": "HoudiniNiagara",
 			"Type": "Runtime",
 			"LoadingPhase": "Default",
-            "WhitelistPlatforms" : [ "Win64", "Win32", "PS4", "XboxOne", "Mac" ]
+			"WhitelistPlatforms": [ "Win64", "Win32", "PS4", "XboxOne", "Mac" ]
 		},
 		{
 			"Name": "HoudiniNiagaraEditor",
 			"Type": "Editor",
 			"LoadingPhase": "Default",
-			"WhitelistPlatforms" : [ "Win64", "Win32", "PS4", "XboxOne", "Mac" ]
+			"WhitelistPlatforms": [ "Win64", "Win32", "PS4", "XboxOne", "Mac" ]
 		}
 	],
-        
+
 	"Plugins":
 	[
 		{
 			"Name": "Niagara",
 			"Enabled": true
 		}
-	]  
+	]
 }
\ No newline at end of file
diff --git a/Engine/Plugins/FX/HoudiniNiagara/Source/HoudiniNiagara/Classes/NiagaraDataInterfaceHoudiniCSV.h b/Engine/Plugins/FX/HoudiniNiagara/Source/HoudiniNiagara/Classes/NiagaraDataInterfaceHoudiniCSV.h
--- a/Engine/Plugins/FX/HoudiniNiagara/Source/HoudiniNiagara/Classes/NiagaraDataInterfaceHoudiniCSV.h
+++ b/Engine/Plugins/FX/HoudiniNiagara/Source/HoudiniNiagara/Classes/NiagaraDataInterfaceHoudiniCSV.h
@@ -33,14 +33,14 @@
 
 
 /** Data Interface allowing sampling of Houdini CSV files. */
-UCLASS(EditInlineNew, Category = "Houdini Niagara", meta = (DisplayName = "Houdini CSV File"))
+UCLASS(EditInlineNew, Category = "Houdini Niagara", meta = (DisplayName = "Houdini Array Info"))
 class HOUDININIAGARA_API UNiagaraDataInterfaceHoudiniCSV : public UNiagaraDataInterface
 {
 	GENERATED_UCLASS_BODY()
 public:
 
-	UPROPERTY( EditAnywhere, Category = "Houdini Niagara", meta = (DisplayName = "Houdini CSV File" ) )
-	UHoudiniCSV* CSVFile;
+	UPROPERTY( EditAnywhere, Category = "Houdini Niagara", meta = (DisplayName = "Houdini CSV Asset" ) )
+	UHoudiniCSV* HoudiniCSVAsset;
 
 	//----------------------------------------------------------------------------
 	//UObject Interface
@@ -63,77 +63,93 @@ public:
 
 	// Returns the float value at a given row and column in the CSV file
 	template<typename RowParamType, typename ColParamType>
-	void GetCSVFloatValue(FVectorVMContext& Context);
+	void GetFloatValue(FVectorVMContext& Context);
 
 	/*
 	template<typename RowParamType, typename ColTitleParamType>
-	void GetCSVFloatValueByString(FVectorVMContext& Context);
+	void GetFloatValueByString(FVectorVMContext& Context);
 	*/
 
-	// Returns a Vector3 value for a given row in the csv file
+	// Returns a Vector3 value for a given row in the CSV file
 	template<typename RowParamType, typename ColParamType>
-	void GetCSVVectorValue(FVectorVMContext& Context);
+	void GetVectorValue(FVectorVMContext& Context);
+
+	// Returns a Vector3 value for a given row in the CSV file
+	template<typename RowParamType, typename ColParamType, typename DoSwapParamType, typename DoScaleParamType>
+	void GetVectorValueEx(FVectorVMContext& Context);
 
 	// Returns the positions for a given row in the CSV file
 	template<typename RowParamType>
-	void GetCSVPosition(FVectorVMContext& Context);
+	void GetPosition(FVectorVMContext& Context);
 
 	// Returns the normals for a given row in the CSV file
 	template<typename RowParamType>
-	void GetCSVNormal(FVectorVMContext& Context);
+	void GetNormal(FVectorVMContext& Context);
 
 	// Returns the time for a given row in the CSV file
 	template<typename RowParamType>
-	void GetCSVTime(FVectorVMContext& Context);
+	void GetTime(FVectorVMContext& Context);
 
 	// Returns the velocity for a given row in the CSV file
 	template<typename RowParamType>
-	void GetCSVVelocity(FVectorVMContext& Context);
+	void GetVelocity(FVectorVMContext& Context);
 
 	// Returns the color for a given row in the CSV file
 	template<typename RowParamType>
-	void GetCSVColor(FVectorVMContext& Context);
+	void GetColor(FVectorVMContext& Context);
 
 	// Returns the position and time for a given row in the CSV file
 	template<typename RowParamType>
-	void GetCSVPositionAndTime(FVectorVMContext& Context);
+	void GetPositionAndTime(FVectorVMContext& Context);
 
 	// Returns the number of rows found in the CSV file
-	void GetNumberOfRowsInCSV(FVectorVMContext& Context);
+	void GetNumberOfRows(FVectorVMContext& Context);
 
 	// Returns the number of columns found in the CSV file
-	void GetNumberOfColumnsInCSV(FVectorVMContext& Context);
+	void GetNumberOfColumns(FVectorVMContext& Context);
 
 	// Returns the number of points found in the CSV file
-	void GetNumberOfParticlesInCSV(FVectorVMContext& Context);
+	void GetNumberOfPoints(FVectorVMContext& Context);
 
-	// Returns the last index of the particles that should be spawned at time t
+	// Returns the last index of the points that should be spawned at time t
 	template<typename TimeParamType>
 	void GetLastRowIndexAtTime(FVectorVMContext& Context);
 
-	// Returns the indexes (min, max) and number of particles that should be spawned at time t
+	// Returns the indexes (min, max) and number of points that should be spawned at time t
 	template<typename TimeParamType>
-	void GetParticleIndexesToSpawnAtTime(FVectorVMContext& Context);
+	void GetPointIDsToSpawnAtTime(FVectorVMContext& Context);
+
+	// Returns the position for a given point at a given time
+	template<typename PointIDParamType, typename TimeParamType>
+	void GetPointPositionAtTime(FVectorVMContext& Context);
+
+	// Returns a float value for a given point at a given time
+	template<typename PointIDParamType, typename ColParamType, typename TimeParamType>
+	void GetPointValueAtTime(FVectorVMContext& Context);
+
+	// Returns a Vector value for a given point at a given time
+	template<typename PointIDParamType, typename ColParamType, typename TimeParamType>
+	void GetPointVectorValueAtTime(FVectorVMContext& Context);
 
-	// Returns the position for a given particle at a given time
-	template<typename NParamType, typename TimeParamType>
-	void GetParticlePositionAtTime(FVectorVMContext& Context);
+	// Returns a Vector value for a given point at a given time
+	template<typename PointIDParamType, typename ColParamType, typename TimeParamType, typename DoSwapParamType, typename DoScaleParamType>
+	void GetPointVectorValueAtTimeEx(FVectorVMContext& Context);
 
-	// Returns a float value for a given particle at a given time
-	template<typename NParamType, typename ColParamType, typename TimeParamType>
-	void GetParticleValueAtTime(FVectorVMContext& Context);
+	// Returns the line indexes (previous, next) for reading values for a given point at a given time
+	template<typename PointIDParamType, typename TimeParamType>
+	void GetRowIndexesForPointAtTime(FVectorVMContext& Context);
 
-	// Returns a Vector value for a given particle at a given time
-	template<typename NParamType, typename ColParamType, typename TimeParamType>
-	void GetParticleVectorValueAtTime(FVectorVMContext& Context);
+	// Return the life value for a given point
+	template<typename PointIDParamType>
+	void GetPointLife(FVectorVMContext& Context);
 
-	// Returns the line indexes (previous, next) for reading values for a given particle at a given time
-	template<typename NParamType, typename TimeParamType>
-	void GetRowIndexesForParticleAtTime(FVectorVMContext& Context);
+	// Return the life of a given point at a given time
+	//template<typename PointIDParamType, typename TimeParamType>
+	//void GetPointLifeAtTime(FVectorVMContext& Context);
 
-	// Return the life of a given particle at a given time
-	template<typename NParamType, typename TimeParamType>
-	void GetParticleLifeAtTime(FVectorVMContext& Context);
+	// Return the type value for a given point
+	template<typename PointIDParamType>
+	void GetPointType( FVectorVMContext& Context );
 	
 	//----------------------------------------------------------------------------
 	// GPU / HLSL Functions
@@ -151,7 +167,7 @@ protected:
 	UPROPERTY()
 	bool GPUBufferDirty;
 
-	// Last Spawned Particle ID
+	// Last Spawned PointID
 	UPROPERTY()
-	int32 LastSpawnedParticleID;
+	int32 LastSpawnedPointID;
 };
diff --git a/Engine/Plugins/FX/HoudiniNiagara/Source/HoudiniNiagara/Private/HoudiniCSV.cpp b/Engine/Plugins/FX/HoudiniNiagara/Source/HoudiniNiagara/Private/HoudiniCSV.cpp
--- a/Engine/Plugins/FX/HoudiniNiagara/Source/HoudiniNiagara/Private/HoudiniCSV.cpp
+++ b/Engine/Plugins/FX/HoudiniNiagara/Source/HoudiniNiagara/Private/HoudiniCSV.cpp
@@ -75,9 +75,9 @@ struct FHoudiniCSVSortPredicate
  
 UHoudiniCSV::UHoudiniCSV( const FObjectInitializer& ObjectInitializer )
 	: Super( ObjectInitializer ),
-	NumberOfLines( -1 ),
+	NumberOfRows( -1 ),
 	NumberOfColumns( -1 ),
-	NumberOfParticles( -1 )
+	NumberOfPoints( -1 )
 {
 	PositionColumnIndex = INDEX_NONE;
 	NormalColumnIndex = INDEX_NONE;
@@ -88,6 +88,9 @@ UHoudiniCSV::UHoudiniCSV( const FObjectInitializer& ObjectInitializer )
 	ColorColumnIndex = INDEX_NONE;
 	AlphaColumnIndex = INDEX_NONE;
 	VelocityColumnIndex = INDEX_NONE;
+	TypeColumnIndex = INDEX_NONE;
+
+	UseCustomTitleRow = false;
 }
 
 void UHoudiniCSV::SetFileName( const FString& TheFileName )
@@ -119,8 +122,8 @@ bool UHoudiniCSV::UpdateFromStringArray( TArray<FString>& RawStringArray )
 {
     // Reset the CSV sizes
     NumberOfColumns = 0;
-    NumberOfLines = 0;
-	NumberOfParticles = 0;
+    NumberOfRows = 0;
+	NumberOfPoints = 0;
 
     // Reset the position, normal and time indexes
     PositionColumnIndex = INDEX_NONE;
@@ -132,28 +135,39 @@ bool UHoudiniCSV::UpdateFromStringArray( TArray<FString>& RawStringArray )
 	ColorColumnIndex = INDEX_NONE;
 	AlphaColumnIndex = INDEX_NONE;
 	VelocityColumnIndex = INDEX_NONE;
+	TypeColumnIndex = INDEX_NONE;
 
     if ( RawStringArray.Num() <= 0 )
     {
-		UE_LOG( LogHoudiniNiagara, Error, TEXT( "Could not load the CSV file, error: not enough lines." ) );
+		UE_LOG( LogHoudiniNiagara, Error, TEXT( "Could not load the CSV file, error: not enough rows in the file." ) );
 		return false;
     }
 
-    // Remove empty lines from the CSV
+    // Remove empty rows from the CSV
     RawStringArray.RemoveAll( [&]( const FString& InString ) { return InString.IsEmpty(); } );
 
-    // Number of lines in the CSV (ignoring the title line)
-    NumberOfLines = RawStringArray.Num() - 1;
-    if ( NumberOfLines < 1 )
+    // Number of rows in the CSV (ignoring the title row)
+    NumberOfRows = RawStringArray.Num() - 1;
+    if ( NumberOfRows < 1 )
     {
-		UE_LOG( LogHoudiniNiagara, Error, TEXT( "Could not load the CSV file, error: not enough lines." ) );
+		UE_LOG( LogHoudiniNiagara, Error, TEXT( "Could not load the CSV file, error: not enough rows in the file." ) );
 		return false;
     }
 
-	// Parses the CSV file's title line to update the column indexes of special values we're interested in
-	// Also look for packed vectors in the first line and update the indexes accordingly
+	// See if we need to use a custom title row
+	// The custom title row will be ignored if it is empty or only composed of spaces
+	FString TitleRow = SourceTitleRow;
+	TitleRow.ReplaceInline(TEXT(" "), TEXT(""));
+	if ( TitleRow.IsEmpty() )
+		UseCustomTitleRow = false;
+
+	if ( !UseCustomTitleRow )
+		SourceTitleRow = RawStringArray[0];
+
+	// Parses the CSV file's title row to update the column indexes of special values we're interested in
+	// Also look for packed vectors in the first row and update the indexes accordingly
 	bool HasPackedVectors = false;
-	if ( !ParseCSVTitleLine( RawStringArray[0], RawStringArray[1], HasPackedVectors ) )
+	if ( !ParseCSVTitleRow( SourceTitleRow, RawStringArray[1], HasPackedVectors ) )
 		return false;
     
     // Remove the title row now that it's been processed
@@ -161,34 +175,34 @@ bool UHoudiniCSV::UpdateFromStringArray( TArray<FString>& RawStringArray )
 
 	// Parses each string of the csv file to a string array
 	TArray< TArray< FString > > ParsedStringArrays;
-	ParsedStringArrays.SetNum( NumberOfLines );
-	for ( int32 rowIdx = 0; rowIdx < NumberOfLines; rowIdx++ )
+	ParsedStringArrays.SetNum( NumberOfRows );
+	for ( int32 rowIdx = 0; rowIdx < NumberOfRows; rowIdx++ )
 	{
-		// Get the current line
-		FString CurrentLine = RawStringArray[ rowIdx ];
+		// Get the current row
+		FString CurrentRow = RawStringArray[ rowIdx ];
 		if ( HasPackedVectors )
 		{
-			// Clean up the packing characters: ()" from the line so it can be parsed properly
-			CurrentLine.ReplaceInline( TEXT("("), TEXT("") );
-			CurrentLine.ReplaceInline( TEXT(")"), TEXT("") );
-			CurrentLine.ReplaceInline( TEXT("\""), TEXT("") );
+			// Clean up the packing characters: ()" from the row so it can be parsed properly
+			CurrentRow.ReplaceInline( TEXT("("), TEXT("") );
+			CurrentRow.ReplaceInline( TEXT(")"), TEXT("") );
+			CurrentRow.ReplaceInline( TEXT("\""), TEXT("") );
 		}
 
-		// Parse the current line to an array
-		TArray<FString> CurrentParsedLine;
-		CurrentLine.ParseIntoArray( CurrentParsedLine, TEXT(",") );
+		// Parse the current row to an array
+		TArray<FString> CurrentParsedRow;
+		CurrentRow.ParseIntoArray( CurrentParsedRow, TEXT(",") );
 
-		// Check the parsed line and number of columns match
-		if ( NumberOfColumns != CurrentParsedLine.Num() )
+		// Check that the parsed row and number of columns match
+		if ( NumberOfColumns != CurrentParsedRow.Num() )
 			UE_LOG( LogHoudiniNiagara, Warning,
-			TEXT("Error while parsing the CSV File. Line %d has %d values instead of the expected %d!"),
-			rowIdx + 1, CurrentParsedLine.Num(), NumberOfColumns );
+			TEXT("Error while parsing the CSV File. Row %d has %d values instead of the expected %d!"),
+			rowIdx + 1, CurrentParsedRow.Num(), NumberOfColumns );
 
-		// Store the parsed line
-		ParsedStringArrays[ rowIdx ] = CurrentParsedLine;
+		// Store the parsed row
+		ParsedStringArrays[ rowIdx ] = CurrentParsedRow;
 	}
 
-	// If we have time values, we have to make sure the lines are sorted by time
+	// If we have time values, we have to make sure the csv rows are sorted by time
 	if ( TimeColumnIndex != INDEX_NONE )
 	{
 		// First check if we need to sort the array
@@ -217,36 +231,37 @@ bool UHoudiniCSV::UpdateFromStringArray( TArray<FString>& RawStringArray )
 
 		if ( NeedToSort )
 		{
-			// We need to sort the CSV lines by their time values
+			// We need to sort the CSV rows by their time values
 			ParsedStringArrays.Sort<FHoudiniCSVSortPredicate>( FHoudiniCSVSortPredicate( TimeColumnIndex, IDColumnIndex ) );
 		}
 	}
 
     // Initialize our different buffers
     FloatCSVData.Empty();
-    FloatCSVData.SetNumZeroed( NumberOfLines * NumberOfColumns );
-    
-    StringCSVData.Empty();
-    StringCSVData.SetNumZeroed( NumberOfLines * NumberOfColumns );
+    FloatCSVData.SetNumZeroed( NumberOfRows * NumberOfColumns );
 
+	/*
+    StringCSVData.Empty();
+    StringCSVData.SetNumZeroed( NumberOfRows * NumberOfColumns );
+	*/
 	// Due to the way that some of the DI functions work,
-	// we expect that the particle IDs start at zero, and increment as the particles are spawned
-	// Make sure this is the case by converting the particle IDs as we read them
-	int32 NextParticleID = 0;
+	// we expect that the point IDs start at zero, and increment as the points are spawned
+	// Make sure this is the case by converting the point IDs as we read them
+	int32 NextPointID = 0;
 	TMap<float, int32> HoudiniIDToNiagaraIDMap;
 
 	// We also keep track of the row indexes for each time values
 	//float lastTimeValue = 0.0;
 	//TimeValuesIndexes.Empty();
 
-	// And the row indexes for each particle
-	ParticleValueIndexes.Empty();
+	// And the row indexes for each point
+	PointValueIndexes.Empty();
 
     // Extract all the values from the table to the float & string buffers
-    TArray<FString> CurrentParsedLine;
+    TArray<FString> CurrentParsedRow;
     for ( int rowIdx = 0; rowIdx < ParsedStringArrays.Num(); rowIdx++ )
     {
-		CurrentParsedLine = ParsedStringArrays[ rowIdx ];
+		CurrentParsedRow = ParsedStringArrays[ rowIdx ];
 
 		// Store the CSV Data in the buffers
 		// The data is stored transposed in those buffers
@@ -255,39 +270,39 @@ bool UHoudiniCSV::UpdateFromStringArray( TArray<FString>& RawStringArray )
 		{
 			// Get the string value for the current column
 			FString CurrentVal = TEXT("0");
-			if ( CurrentParsedLine.IsValidIndex( colIdx ) )
+			if ( CurrentParsedRow.IsValidIndex( colIdx ) )
 			{
-				CurrentVal = CurrentParsedLine[ colIdx ];
+				CurrentVal = CurrentParsedRow[ colIdx ];
 			}
 			else
 			{
 				UE_LOG( LogHoudiniNiagara, Warning,
-				TEXT("Error while parsing the CSV File. Line %d has an invalid value for column %d!"),
+				TEXT("Error while parsing the CSV File. Row %d has an invalid value for column %d!"),
 				rowIdx + 1, colIdx + 1 );
 			}
 
 			// Convert the string value to a float
 			float FloatValue = FCString::Atof( *CurrentVal );
 
-			// Handle particles IDs here
+			// Handle point IDs here
 			if ( colIdx == IDColumnIndex )
 			{
-				// The particle ID may need to be replaced
+				// The point ID may need to be replaced
 				if ( !HoudiniIDToNiagaraIDMap.Contains( FloatValue ) )
 				{
-					// We found a new particle, add it to the ID map
-					HoudiniIDToNiagaraIDMap.Add( FloatValue, NextParticleID++ );
+					// We found a new point, so we add it to the ID map
+					HoudiniIDToNiagaraIDMap.Add( FloatValue, NextPointID++ );
 
-					// Add a new array for that particle's indexes
-					ParticleValueIndexes.Add( FParticleIndexes() );
+					// Add a new array for that point's indexes
+					PointValueIndexes.Add( FPointIndexes() );
 				}
 
 				// Get the Niagara ID from this Houdini ID
 				CurrentID = HoudiniIDToNiagaraIDMap[ FloatValue ];
 				FloatValue = (float)CurrentID;
 
-				// Add the current row to this particle's row index list
-				ParticleValueIndexes[ CurrentID ].RowIndexes.Add( rowIdx );
+				// Add the current row to this point's row index list
+				PointValueIndexes[ CurrentID ].RowIndexes.Add( rowIdx );
 			}
 
 			/*
@@ -303,47 +318,52 @@ bool UHoudiniCSV::UpdateFromStringArray( TArray<FString>& RawStringArray )
 			*/
 
 			// Store the Value in the buffer
-			FloatCSVData[ rowIdx + ( colIdx * NumberOfLines ) ] = FloatValue;
+			FloatCSVData[ rowIdx + ( colIdx * NumberOfRows ) ] = FloatValue;
 
+			/*
 			// Keep the original string value in a buffer too
-			StringCSVData[ rowIdx + ( colIdx * NumberOfLines ) ] = CurrentVal;
+			StringCSVData[ rowIdx + ( colIdx * NumberOfRows ) ] = CurrentVal;
+			*/
 		}
 
-		// Look at the particle ID, the max ID will be our number of particles
-		if ( NumberOfParticles <= CurrentID )
-			NumberOfParticles = CurrentID + 1;
+		// Look at the point ID, the max ID will be our number of points
+		if ( NumberOfPoints <= CurrentID )
+			NumberOfPoints = CurrentID + 1;
     }
 	
-	NumberOfParticles = HoudiniIDToNiagaraIDMap.Num();
-	if ( NumberOfParticles <= 0 )
-		NumberOfParticles = NumberOfLines;
+	NumberOfPoints = HoudiniIDToNiagaraIDMap.Num();
+	if ( NumberOfPoints <= 0 )
+		NumberOfPoints = NumberOfRows;
 
-	// Look for particle specific attributes
+	// Look for point specific attributes
 	SpawnTimes.Empty();
-	SpawnTimes.Init( -1.0f, NumberOfParticles );
+	SpawnTimes.Init( -1.0f, NumberOfPoints );
 
 	LifeValues.Empty();
-	LifeValues.Init( -1.0f,  NumberOfParticles );
+	LifeValues.Init( -1.0f,  NumberOfPoints );
+
+	PointTypes.Empty();
+	PointTypes.Init( -1, NumberOfPoints );
 
 	float MaxTime = -1.0f;
-	for ( int rowIdx = 0; rowIdx < NumberOfLines; rowIdx++ )
+	for ( int rowIdx = 0; rowIdx < NumberOfRows; rowIdx++ )
 	{
-		// Get the particle ID
+		// Get the point ID
 		int32 CurrentID = rowIdx;
 		if ( IDColumnIndex != INDEX_NONE )
-			CurrentID = (int32)FloatCSVData[ rowIdx + ( IDColumnIndex * NumberOfLines ) ];
+			CurrentID = (int32)FloatCSVData[ rowIdx + ( IDColumnIndex * NumberOfRows ) ];
 
-		// Get the time value for the current line
+		// Get the time value for the current row
 		float CurrentTime = 0.0f;
 		if ( TimeColumnIndex != INDEX_NONE )
-			CurrentTime = FloatCSVData[ rowIdx + ( TimeColumnIndex * NumberOfLines ) ];
+			CurrentTime = FloatCSVData[ rowIdx + ( TimeColumnIndex * NumberOfRows ) ];
 
 		if ( LifeColumnIndex != INDEX_NONE )
-		{
-			// Set spawn time and life using life values
-			float CurrentLife = FloatCSVData[ rowIdx + ( LifeColumnIndex * NumberOfLines ) ];
+		{	
 			if ( SpawnTimes[ CurrentID ] < 0.0f )
 			{
+				// Set spawn time and life using life values
+				float CurrentLife = FloatCSVData [rowIdx + (LifeColumnIndex * NumberOfRows) ];
 				SpawnTimes[ CurrentID ] = CurrentTime;
 				LifeValues[ CurrentID ] = CurrentLife;
 			}
@@ -351,10 +371,10 @@ bool UHoudiniCSV::UpdateFromStringArray( TArray<FString>& RawStringArray )
 		else if ( AliveColumnIndex != INDEX_NONE )
 		{
 			// Set spawn time and life using the alive bool values
-			bool CurrentAlive = ( FloatCSVData[ rowIdx + ( AliveColumnIndex * NumberOfLines ) ] == 1.0f );
+			bool CurrentAlive = ( FloatCSVData[ rowIdx + ( AliveColumnIndex * NumberOfRows ) ] == 1.0f );
 			if ( ( SpawnTimes[ CurrentID ] < 0.0f ) && CurrentAlive )
 			{
-				// Spawn time is when the particle is first seen alive
+				// Spawn time is when the point is first seen alive
 				SpawnTimes[ CurrentID ] = CurrentTime;
 			}
 			else if ( ( SpawnTimes[ CurrentID ] >= 0.0f ) && !CurrentAlive )
@@ -365,21 +385,31 @@ bool UHoudiniCSV::UpdateFromStringArray( TArray<FString>& RawStringArray )
 		}
 		else
 		{
-			// No life or alive value, spawn time is the first time we see the particle
+			// No life or alive value, spawn time is the first time we have a value for this point
 			if ( SpawnTimes [ CurrentID ] == INDEX_NONE )
 				SpawnTimes[ CurrentID ] = CurrentTime;
 		}
+
+		// Keep track of the point type at spawn
+		if ( PointTypes[ CurrentID ] < 0 )
+		{
+			float CurrentType = 0.0f;
+			if (TypeColumnIndex != INDEX_NONE)
+				GetFloatValue( rowIdx, TypeColumnIndex, CurrentType );
+
+			PointTypes[ CurrentID ] = (int32)CurrentType;
+		}
+
 	}
     return true;
 }
 
-bool UHoudiniCSV::ParseCSVTitleLine( const FString& TitleLine, const FString& FirstLine, bool& HasPackedVectors )
+bool UHoudiniCSV::ParseCSVTitleRow( const FString& TitleRow, const FString& FirstValueRow, bool& HasPackedVectors )
 {
-	// Get the number of values per lines via the title line
-    //FString TitleString = RawStringArray[ 0 ];
-    TitleRowArray.Empty();
-    TitleLine.ParseIntoArray( TitleRowArray, TEXT(",") );
-    NumberOfColumns = TitleRowArray.Num();
+	// Get the number of values per row via the title row
+    ColumnTitleArray.Empty();
+    TitleRow.ParseIntoArray( ColumnTitleArray, TEXT(",") );
+    NumberOfColumns = ColumnTitleArray.Num();
     if ( NumberOfColumns < 1 )
     {
 		UE_LOG( LogHoudiniNiagara, Error, TEXT( "Could not load the CSV file, error: not enough columns." ) );
@@ -387,12 +417,12 @@ bool UHoudiniCSV::ParseCSVTitleLine( const FString& TitleLine, const FString& Fi
     }
 
     // Look for the position, normal and time attributes indexes
-    for ( int32 n = 0; n < TitleRowArray.Num(); n++ )
+    for ( int32 n = 0; n < ColumnTitleArray.Num(); n++ )
     {
 		// Remove spaces from the title row
-		TitleRowArray[ n ].ReplaceInline( TEXT(" "), TEXT("") );
+		ColumnTitleArray[ n ].ReplaceInline( TEXT(" "), TEXT("") );
 
-		FString CurrentTitle = TitleRowArray[ n ];
+		FString CurrentTitle = ColumnTitleArray[ n ];
 		if ( CurrentTitle.Equals( TEXT("P"), ESearchCase::IgnoreCase )
 			|| CurrentTitle.Equals( TEXT("Px"), ESearchCase::IgnoreCase )
 			|| CurrentTitle.Equals( TEXT("X"), ESearchCase::IgnoreCase )
@@ -447,33 +477,38 @@ bool UHoudiniCSV::ParseCSVTitleLine( const FString& TitleLine, const FString& Fi
 			if ( VelocityColumnIndex == INDEX_NONE )
 				VelocityColumnIndex = n;
 		}
+		else if ( CurrentTitle.Equals(TEXT("type"), ESearchCase::IgnoreCase ) )
+		{
+			if ( TypeColumnIndex == INDEX_NONE )
+				TypeColumnIndex = n;
+		}
     }
 
-	// Read the first line of the CSV file, and look for packed vectors value (X,Y,Z)
+	// Read the first row of the CSV file, and look for packed vectors value (X,Y,Z)
     // We'll have to expand them in the title row to match the parsed data
 	HasPackedVectors = false;
 	int32 FoundPackedVectorCharIndex = 0;    
     while ( FoundPackedVectorCharIndex != INDEX_NONE )
     {
-		// Try to find ( in the line
-		FoundPackedVectorCharIndex = FirstLine.Find( TEXT("("), ESearchCase::IgnoreCase, ESearchDir::FromStart, FoundPackedVectorCharIndex );
+		// Try to find ( in the row
+		FoundPackedVectorCharIndex = FirstValueRow.Find( TEXT("("), ESearchCase::IgnoreCase, ESearchDir::FromStart, FoundPackedVectorCharIndex );
 		if ( FoundPackedVectorCharIndex == INDEX_NONE )
 			break;
 
 		// We want to know which column this char belong to
 		int32 FoundPackedVectorColumnIndex = INDEX_NONE;
 		{
-			// Chop the first line up to the found character
-			FString FirstLineLeft = FirstLine.Left( FoundPackedVectorCharIndex );
+			// Chop the first row up to the found character
+			FString FirstRowLeft = FirstValueRow.Left( FoundPackedVectorCharIndex );
 
 			// ReplaceInLine returns the number of occurences of ",", that's what we want! 
-			FoundPackedVectorColumnIndex = FirstLineLeft.ReplaceInline(TEXT(","), TEXT(""));
+			FoundPackedVectorColumnIndex = FirstRowLeft.ReplaceInline(TEXT(","), TEXT(""));
 		}
 
-		if ( !TitleRowArray.IsValidIndex( FoundPackedVectorColumnIndex ) )
+		if ( !ColumnTitleArray.IsValidIndex( FoundPackedVectorColumnIndex ) )
 		{
 			UE_LOG( LogHoudiniNiagara, Warning,
-			TEXT( "Error while parsing the CSV File. Couldn't unpack vector found at character %d in the first line!" ),
+			TEXT( "Error while parsing the CSV File. Couldn't unpack vector found at character %d in the first row!" ),
 			FoundPackedVectorCharIndex + 1 );
 			continue;
 		}
@@ -482,8 +517,8 @@ bool UHoudiniCSV::ParseCSVTitleLine( const FString& TitleLine, const FString& Fi
 		int32 FoundVectorSize = 0;
 		{
 			// Extract the vector string
-			int32 FoundPackedVectorEndCharIndex = FirstLine.Find( TEXT(")"), ESearchCase::IgnoreCase, ESearchDir::FromStart, FoundPackedVectorCharIndex );
-			FString VectorString = FirstLine.Mid( FoundPackedVectorCharIndex + 1, FoundPackedVectorEndCharIndex - FoundPackedVectorCharIndex - 1 );
+			int32 FoundPackedVectorEndCharIndex = FirstValueRow.Find( TEXT(")"), ESearchCase::IgnoreCase, ESearchDir::FromStart, FoundPackedVectorCharIndex );
+			FString VectorString = FirstValueRow.Mid( FoundPackedVectorCharIndex + 1, FoundPackedVectorEndCharIndex - FoundPackedVectorCharIndex - 1 );
 
 			// Use ReplaceInLine to count the number of , to get the vector's size!
 			FoundVectorSize = VectorString.ReplaceInline( TEXT(","), TEXT("") ) + 1;
@@ -499,35 +534,35 @@ bool UHoudiniCSV::ParseCSVTitleLine( const FString& TitleLine, const FString& Fi
 		if ( ( FoundPackedVectorColumnIndex == PositionColumnIndex ) && ( FoundVectorSize == 3 ) )
 		{
 			// Expand P to Px,Py,Pz
-			TitleRowArray[ PositionColumnIndex ] = TEXT("Px");
-			TitleRowArray.Insert( TEXT( "Py" ), PositionColumnIndex + 1 );
-			TitleRowArray.Insert( TEXT( "Pz" ), PositionColumnIndex + 2 );
+			ColumnTitleArray[ PositionColumnIndex ] = TEXT("Px");
+			ColumnTitleArray.Insert( TEXT( "Py" ), PositionColumnIndex + 1 );
+			ColumnTitleArray.Insert( TEXT( "Pz" ), PositionColumnIndex + 2 );
 		}
 		else if ( ( FoundPackedVectorColumnIndex == NormalColumnIndex ) && ( FoundVectorSize == 3 ) )
 		{
 			// Expand N to Nx,Ny,Nz
-			TitleRowArray[ NormalColumnIndex ] = TEXT("Nx");
-			TitleRowArray.Insert( TEXT("Ny"), NormalColumnIndex + 1 );
-			TitleRowArray.Insert( TEXT("Nz"), NormalColumnIndex + 2 );
+			ColumnTitleArray[ NormalColumnIndex ] = TEXT("Nx");
+			ColumnTitleArray.Insert( TEXT("Ny"), NormalColumnIndex + 1 );
+			ColumnTitleArray.Insert( TEXT("Nz"), NormalColumnIndex + 2 );
 		}
 		else if ( ( FoundPackedVectorColumnIndex == VelocityColumnIndex ) && ( FoundVectorSize == 3 ) )
 		{
 			// Expand V to Vx,Vy,Vz
-			TitleRowArray[ VelocityColumnIndex ] = TEXT("Vx");
-			TitleRowArray.Insert( TEXT("Vy"), VelocityColumnIndex + 1 );
-			TitleRowArray.Insert( TEXT("Vz"), VelocityColumnIndex + 2 );
+			ColumnTitleArray[ VelocityColumnIndex ] = TEXT("Vx");
+			ColumnTitleArray.Insert( TEXT("Vy"), VelocityColumnIndex + 1 );
+			ColumnTitleArray.Insert( TEXT("Vz"), VelocityColumnIndex + 2 );
 		}
 		else if ( ( FoundPackedVectorColumnIndex == ColorColumnIndex ) && ( ( FoundVectorSize == 3 ) || ( FoundVectorSize == 4 ) ) )
 		{
 			// Expand Cd to R, G, B 
-			TitleRowArray[ ColorColumnIndex ] = TEXT("R");
-			TitleRowArray.Insert( TEXT("G"), ColorColumnIndex + 1 );
-			TitleRowArray.Insert( TEXT("B"), ColorColumnIndex + 2 );
+			ColumnTitleArray[ ColorColumnIndex ] = TEXT("R");
+			ColumnTitleArray.Insert( TEXT("G"), ColorColumnIndex + 1 );
+			ColumnTitleArray.Insert( TEXT("B"), ColorColumnIndex + 2 );
 
 			if ( FoundVectorSize == 4 )
 			{
 				// Insert A if we had RGBA
-				TitleRowArray.Insert( TEXT("A"), ColorColumnIndex + 3 );
+				ColumnTitleArray.Insert( TEXT("A"), ColorColumnIndex + 3 );
 				if ( AlphaColumnIndex == INDEX_NONE )
 					AlphaColumnIndex = ColorColumnIndex + 3;
 			}
@@ -535,11 +570,11 @@ bool UHoudiniCSV::ParseCSVTitleLine( const FString& TitleLine, const FString& Fi
 		else
 		{
 			// Expand the vector's title from V to V, V1, V2, V3 ...
-			FString FoundPackedVectortTitle = TitleRowArray[ FoundPackedVectorColumnIndex ];
+			FString FoundPackedVectortTitle = ColumnTitleArray[ FoundPackedVectorColumnIndex ];
 			for ( int32 n = 1; n < FoundVectorSize; n++ )
 			{
 				FString CurrentTitle = FoundPackedVectortTitle + FString::FromInt( n );
-				TitleRowArray.Insert( CurrentTitle, FoundPackedVectorColumnIndex + n );
+				ColumnTitleArray.Insert( CurrentTitle, FoundPackedVectorColumnIndex + n );
 			}
 		}
 
@@ -571,41 +606,54 @@ bool UHoudiniCSV::ParseCSVTitleLine( const FString& TitleLine, const FString& Fi
 		if ( VelocityColumnIndex != INDEX_NONE && ( VelocityColumnIndex > FoundPackedVectorColumnIndex ) )
 			VelocityColumnIndex += FoundVectorSize - 1;
 
+		if ( TypeColumnIndex != INDEX_NONE && ( TypeColumnIndex > FoundPackedVectorColumnIndex ) )
+			TypeColumnIndex += FoundVectorSize - 1;
+
 		HasPackedVectors = true;
 		FoundPackedVectorCharIndex++;
     }
 
-    // For sanity, Check that the number of columns matches the title row and the first line
+    // For sanity, Check that the number of columns matches the title row and the first row
     {
 		// Check the title row
-		if ( NumberOfColumns != TitleRowArray.Num() )
+		if ( NumberOfColumns != ColumnTitleArray.Num() )
 			UE_LOG( LogHoudiniNiagara, Error,
 			TEXT( "Error while parsing the CSV File. Found %d columns but the Title string has %d values! Some values will have an offset!" ),
-			NumberOfColumns, TitleRowArray.Num() );
+			NumberOfColumns, ColumnTitleArray.Num() );
 
-		// Use ReplaceInLine to count the number of columns in the first line and make sure it's correct
-		FString FirstLineCopy = FirstLine;
-		int32 FirstLineColumns = FirstLineCopy.ReplaceInline( TEXT(","), TEXT("") ) + 1;
-		if ( NumberOfColumns != FirstLineColumns )
+		// Use ReplaceInLine to count the number of columns in the first row and make sure it's correct
+		FString FirstValueRowCopy = FirstValueRow;
+		int32 FirstRowColumnCount = FirstValueRowCopy.ReplaceInline( TEXT(","), TEXT("") ) + 1;
+		if ( NumberOfColumns != FirstRowColumnCount )
 			UE_LOG( LogHoudiniNiagara, Error,
-			TEXT("Error while parsing the CSV File. Found %d columns but found %d values in the first line! Some values will have an offset!" ),
-			NumberOfColumns, FirstLineColumns );
+			TEXT("Error while parsing the CSV File. Found %d columns but found %d values in the first row! Some values will have an offset!" ),
+			NumberOfColumns, FirstRowColumnCount );
     }
+	/*
+	// Update the TitleRow uproperty
+	TitleRow.Empty();
+	for (int32 n = 0; n < TitleRowArray.Num(); n++)
+	{
+		TitleRow += TitleRowArray[n];
+		if ( n < TitleRowArray.Num() - 1 )
+			TitleRow += TEXT(",");
+	}
+	*/	
 
 	return true;
 }
 
 
 // Returns the float value at a given point in the CSV file
-bool UHoudiniCSV::GetCSVFloatValue( const int32& lineIndex, const int32& colIndex, float& value )
+bool UHoudiniCSV::GetFloatValue( const int32& rowIndex, const int32& colIndex, float& value )
 {
-    if ( lineIndex < 0 || lineIndex >= NumberOfLines )
+    if ( rowIndex < 0 || rowIndex >= NumberOfRows )
 		return false;
 
     if ( colIndex < 0 || colIndex >= NumberOfColumns )
 		return false;
 
-    int32 Index = lineIndex + ( colIndex * NumberOfLines );
+    int32 Index = rowIndex + ( colIndex * NumberOfRows );
     if ( FloatCSVData.IsValidIndex( Index ) )
     {
 		value = FloatCSVData[ Index ];
@@ -615,16 +663,17 @@ bool UHoudiniCSV::GetCSVFloatValue( const int32& lineIndex, const int32& colInde
     return false;
 }
 
+/*
 // Returns the float value at a given point in the CSV file
-bool UHoudiniCSV::GetCSVStringValue( const int32& lineIndex, const int32& colIndex, FString& value )
+bool UHoudiniCSV::GetCSVStringValue( const int32& rowIndex, const int32& colIndex, FString& value )
 {
-    if ( lineIndex < 0 || lineIndex >= NumberOfLines )
+    if ( rowIndex < 0 || rowIndex >= NumberOfRows )
 		return false;
 
     if ( colIndex < 0 || colIndex >= NumberOfColumns )
 		return false;
 
-    int32 Index = lineIndex + ( colIndex * NumberOfLines );
+    int32 Index = rowIndex + ( colIndex * NumberOfRows );
     if ( StringCSVData.IsValidIndex( Index ) )
     {
 		value = StringCSVData[ Index ];
@@ -633,18 +682,19 @@ bool UHoudiniCSV::GetCSVStringValue( const int32& lineIndex, const int32& colInd
 
     return false;
 }
+*/
 
 // Returns a Vector 3 for a given point in the CSV file
-bool UHoudiniCSV::GetCSVVectorValue( const int32& lineIndex, const int32& colIndex, FVector& value, const bool& DoSwap, const bool& DoScale )
+bool UHoudiniCSV::GetVectorValue( const int32& rowIndex, const int32& colIndex, FVector& value, const bool& DoSwap, const bool& DoScale )
 {
     FVector V = FVector::ZeroVector;
-    if ( !GetCSVFloatValue( lineIndex, colIndex, V.X ) )
+    if ( !GetFloatValue( rowIndex, colIndex, V.X ) )
 		return false;
 
-    if ( !GetCSVFloatValue( lineIndex, colIndex + 1, V.Y ) )
+    if ( !GetFloatValue( rowIndex, colIndex + 1, V.Y ) )
 		return false;
 
-    if ( !GetCSVFloatValue( lineIndex, colIndex + 2, V.Z ) )
+    if ( !GetFloatValue( rowIndex, colIndex + 2, V.Z ) )
 		return false;
 
     if ( DoScale )
@@ -662,10 +712,10 @@ bool UHoudiniCSV::GetCSVVectorValue( const int32& lineIndex, const int32& colInd
 }
 
 // Returns a Vector 3 for a given point in the CSV file
-bool UHoudiniCSV::GetCSVPositionValue( const int32& lineIndex, FVector& value )
+bool UHoudiniCSV::GetPositionValue( const int32& rowIndex, FVector& value )
 {
     FVector V = FVector::ZeroVector;
-    if ( !GetCSVVectorValue( lineIndex, PositionColumnIndex, V, true, true ) )
+    if ( !GetVectorValue( rowIndex, PositionColumnIndex, V, true, true ) )
 		return false;
 
     value = V;
@@ -674,10 +724,10 @@ bool UHoudiniCSV::GetCSVPositionValue( const int32& lineIndex, FVector& value )
 }
 
 // Returns a Vector 3 for a given point in the CSV file
-bool UHoudiniCSV::GetCSVNormalValue( const int32& lineIndex, FVector& value )
+bool UHoudiniCSV::GetNormalValue( const int32& rowIndex, FVector& value )
 {
     FVector V = FVector::ZeroVector;
-    if ( !GetCSVVectorValue(lineIndex, NormalColumnIndex, V, true, false ) )
+    if ( !GetVectorValue( rowIndex, NormalColumnIndex, V, true, false ) )
 		return false;
 
     value = V;
@@ -686,10 +736,10 @@ bool UHoudiniCSV::GetCSVNormalValue( const int32& lineIndex, FVector& value )
 }
 
 // Returns a time value for a given point in the CSV file
-bool UHoudiniCSV::GetCSVTimeValue( const int32& lineIndex, float& value )
+bool UHoudiniCSV::GetTimeValue( const int32& rowIndex, float& value )
 {
     float temp;
-    if ( !GetCSVFloatValue( lineIndex, TimeColumnIndex, temp ) )
+    if ( !GetFloatValue( rowIndex, TimeColumnIndex, temp ) )
 		return false;
 
     value = temp;
@@ -698,10 +748,10 @@ bool UHoudiniCSV::GetCSVTimeValue( const int32& lineIndex, float& value )
 }
 
 // Returns a Color for a given point in the CSV file
-bool UHoudiniCSV::GetCSVColorValue( const int32& lineIndex, FLinearColor& value )
+bool UHoudiniCSV::GetColorValue( const int32& rowIndex, FLinearColor& value )
 {
 	FVector V = FVector::OneVector;
-	if ( !GetCSVVectorValue( lineIndex, ColorColumnIndex, V, false, false ) )
+	if ( !GetVectorValue( rowIndex, ColorColumnIndex, V, false, false ) )
 		return false;
 
 	FLinearColor C = FLinearColor::White;
@@ -710,7 +760,7 @@ bool UHoudiniCSV::GetCSVColorValue( const int32& lineIndex, FLinearColor& value
 	C.B = V.Z;
 
 	float alpha = 1.0f;
-	if ( GetCSVFloatValue( lineIndex, AlphaColumnIndex, alpha ) )
+	if ( GetFloatValue( rowIndex, AlphaColumnIndex, alpha ) )
 		C.A = alpha;	
 
 	value = C;
@@ -719,10 +769,10 @@ bool UHoudiniCSV::GetCSVColorValue( const int32& lineIndex, FLinearColor& value
 }
 
 // Returns a Velocity Vector3 for a given point in the CSV file
-bool UHoudiniCSV::GetCSVVelocityValue( const int32& lineIndex, FVector& value )
+bool UHoudiniCSV::GetVelocityValue( const int32& rowIndex, FVector& value )
 {
 	FVector V = FVector::ZeroVector;
-	if ( !GetCSVVectorValue( lineIndex, VelocityColumnIndex, V, true, false ) )
+	if ( !GetVectorValue( rowIndex, VelocityColumnIndex, V, true, false ) )
 		return false;
 
 	value = V;
@@ -731,53 +781,52 @@ bool UHoudiniCSV::GetCSVVelocityValue( const int32& lineIndex, FVector& value )
 }
 
 // Returns the number of points found in the CSV file
-int32 UHoudiniCSV::GetNumberOfParticlesInCSV()
+int32 UHoudiniCSV::GetNumberOfPoints()
 {
 	if ( IDColumnIndex != INDEX_NONE )
-		return NumberOfParticles;
+		return NumberOfPoints;
 
-    return NumberOfLines;
+    return NumberOfRows;
 }
 
-// Returns the number of lines found in the CSV file
-int32 UHoudiniCSV::GetNumberOfLinesInCSV()
+// Returns the number of rows found in the CSV file
+int32 UHoudiniCSV::GetNumberOfRows()
 {
-	return NumberOfLines;
+	return NumberOfRows;
 }
 
 // Returns the number of columns found in the CSV file
-int32 UHoudiniCSV::GetNumberOfColumnsInCSV()
+int32 UHoudiniCSV::GetNumberOfColumns()
 {
 	return NumberOfColumns;
 }
 
 // Get the last row index for a given time value (the row with a time smaller or equal to desiredTime)
-// If the CSV file doesn't have time informations, returns false and set the LastRowIndex to the last line in the file
+// If the CSV file doesn't have time informations, returns false and set the LastRowIndex to the last row in the file
 // If desiredTime is smaller than the time value in the first row, LastRowIndex will be set to -1
 // If desiredTime is higher than the last time value in the last row of the csv file, LastIndex will be set to the last row's index
 bool UHoudiniCSV::GetLastRowIndexAtTime(const float& desiredTime, int32& lastRowIndex)
 {
-	// If we dont have proper time info, always return the last line
-	if (TimeColumnIndex < 0 || TimeColumnIndex >= NumberOfColumns)
+	// If we dont have proper time info, always return the last row
+	if ( TimeColumnIndex < 0 || TimeColumnIndex >= NumberOfColumns )
 	{
-		lastRowIndex = NumberOfLines - 1;
+		lastRowIndex = NumberOfRows - 1;
 		return false;
 	}
 
 	float temp_time = 0.0f;
-	// Check first if we have anything to spawn at the current time by looking at the last value
-	if ( GetCSVTimeValue( NumberOfLines - 1, temp_time ) && temp_time < desiredTime )
+	if ( GetTimeValue( NumberOfRows - 1, temp_time ) && temp_time < desiredTime )
 	{
 		// We didn't find a suitable index because the desired time is higher than our last time value
-		lastRowIndex = NumberOfLines - 1;
+		lastRowIndex = NumberOfRows - 1;
 		return true;
 	}
 
-	// Iterates through all the lines
+	// Iterates through all the rows
 	lastRowIndex = INDEX_NONE;
-	for ( int32 n = 0; n < NumberOfLines; n++ )
+	for ( int32 n = 0; n < NumberOfRows; n++ )
 	{
-		if ( GetCSVTimeValue( n, temp_time ) )
+		if ( GetTimeValue( n, temp_time ) )
 		{
 			if ( temp_time == desiredTime )
 			{
@@ -793,40 +842,40 @@ bool UHoudiniCSV::GetLastRowIndexAtTime(const float& desiredTime, int32& lastRow
 
 	// We didn't find a suitable index because the desired time is higher than our last time value
 	if ( lastRowIndex == INDEX_NONE )
-		lastRowIndex = NumberOfLines - 1;
+		lastRowIndex = NumberOfRows - 1;
 
 	return true;
 }
-// Get the last index of the particles with a time value smaller or equal to desiredTime
-// If the CSV file doesn't have time informations, returns false and set the LastIndex to the last particle
-// If desiredTime is smaller than the first particle, LastIndex will be set to -1
-// If desiredTime is higher than the last particle in the csv file, LastIndex will be set to the last particle's index
-bool UHoudiniCSV::GetLastParticleIndexAtTime( const float& desiredTime, int32& lastID )
+// Get the last index of the points with a time value smaller or equal to desiredTime
+// If the CSV file doesn't have time informations, returns false and set the LastIndex to the last point
+// If desiredTime is smaller than the first point time, LastIndex will be set to -1
+// If desiredTime is higher than the last point time in the csv file, LastIndex will be set to the last point's index
+bool UHoudiniCSV::GetLastPointIDToSpawnAtTime( const float& desiredTime, int32& lastID )
 {
-    // If we dont have proper time info, always return the last particle
+    // If we dont have proper time info, always return the last point
     if ( TimeColumnIndex < 0 || TimeColumnIndex >= NumberOfColumns )
     {
-		lastID = NumberOfParticles - 1;
+		lastID = NumberOfPoints - 1;
 		return false;
     }
 
     float temp_time = 0.0f;
-	if ( !SpawnTimes.IsValidIndex( NumberOfParticles - 1 ) )
+	if ( !SpawnTimes.IsValidIndex( NumberOfPoints - 1 ) )
 	{
-		lastID = NumberOfParticles - 1;
+		lastID = NumberOfPoints - 1;
 		return false;
 	}
 
-	if ( SpawnTimes[ NumberOfParticles - 1 ] < desiredTime )
+	if ( SpawnTimes[ NumberOfPoints - 1 ] < desiredTime )
 	{
 		// We didn't find a suitable index because the desired time is higher than our last time value
-		lastID = NumberOfParticles - 1;
+		lastID = NumberOfPoints - 1;
 		return true;
 	}
 
-	// Iterates through all the particles
+	// Iterates through all the points
 	lastID = INDEX_NONE;
-	for ( int32 n = 0; n < NumberOfParticles; n++ )
+	for ( int32 n = 0; n < NumberOfPoints; n++ )
 	{
 		temp_time = SpawnTimes[ n ];
 
@@ -844,39 +893,65 @@ bool UHoudiniCSV::GetLastParticleIndexAtTime( const float& desiredTime, int32& l
 	return true;
 }
 
-bool UHoudiniCSV::GetParticleLifeAtTime( const int32& ParticleID, const float& DesiredTime, float& Value )
+bool UHoudiniCSV::GetPointType(const int32& PointID, int32& Value)
+{
+	if ( !PointTypes.IsValidIndex( PointID ) )
+	{
+		Value = -1;
+		return false;
+	}
+
+	Value = PointTypes[ PointID ];
+
+	return true;
+}
+
+bool UHoudiniCSV::GetPointLife(const int32& PointID, float& Value)
+{
+	if ( !LifeValues.IsValidIndex( PointID ) )
+	{
+		Value = -1.0f;
+		return false;
+	}
+
+	Value = LifeValues[ PointID ];
+
+	return true;
+}
+
+bool UHoudiniCSV::GetPointLifeAtTime( const int32& PointID, const float& DesiredTime, float& Value )
 {
-	if ( !SpawnTimes.IsValidIndex( ParticleID )  || !LifeValues.IsValidIndex( ParticleID ) )
+	if ( !SpawnTimes.IsValidIndex( PointID )  || !LifeValues.IsValidIndex( PointID ) )
 	{
 		Value = -1.0f;
 		return false;
 	}
 
-	if ( DesiredTime < SpawnTimes[ ParticleID ] )
+	if ( DesiredTime < SpawnTimes[ PointID ] )
 	{
-		Value = LifeValues[ ParticleID ];
+		Value = LifeValues[ PointID ];
 	}
 	else
 	{
-		Value = LifeValues[ ParticleID ] - ( DesiredTime - SpawnTimes[ ParticleID ] );
+		Value = LifeValues[ PointID ] - ( DesiredTime - SpawnTimes[ PointID ] );
 	}
 
 	return true;
 }
 
-bool UHoudiniCSV::GetParticleValueAtTime( const int32& ParticleID, const int32& ColumnIndex, const float& desiredTime, float& Value )
+bool UHoudiniCSV::GetPointValueAtTime( const int32& PointID, const int32& ColumnIndex, const float& desiredTime, float& Value )
 {
 	int32 PrevIndex = -1;
 	int32 NextIndex = -1;
 	float PrevWeight = 1.0f;
 
-	if ( !GetParticleLineIndexAtTime( ParticleID, desiredTime, PrevIndex, NextIndex, PrevWeight ) )
+	if ( !GetRowIndexesForPointAtTime( PointID, desiredTime, PrevIndex, NextIndex, PrevWeight ) )
 		return false;
 
 	float PrevValue, NextValue;
-	if ( !GetCSVFloatValue( PrevIndex, ColumnIndex, PrevValue ) )
+	if ( !GetFloatValue( PrevIndex, ColumnIndex, PrevValue ) )
 		return false;
-	if ( !GetCSVFloatValue( NextIndex, ColumnIndex, NextValue ) )
+	if ( !GetFloatValue( NextIndex, ColumnIndex, NextValue ) )
 		return false;
 
 	Value = FMath::Lerp( PrevValue, NextValue, PrevWeight );
@@ -884,10 +959,10 @@ bool UHoudiniCSV::GetParticleValueAtTime( const int32& ParticleID, const int32&
 	return true;
 }
 
-bool UHoudiniCSV::GetParticlePositionAtTime( const int32& ParticleID, const float& desiredTime, FVector& Vector )
+bool UHoudiniCSV::GetPointPositionAtTime( const int32& PointID, const float& desiredTime, FVector& Vector )
 {
 	FVector V = FVector::ZeroVector;
-	if ( !GetParticleVectorValueAtTime(ParticleID, PositionColumnIndex, desiredTime, V, true, true ) )
+	if ( !GetPointVectorValueAtTime(PointID, PositionColumnIndex, desiredTime, V, true, true ) )
 		return false;
 
 	Vector = V;
@@ -895,19 +970,19 @@ bool UHoudiniCSV::GetParticlePositionAtTime( const int32& ParticleID, const floa
 	return true;
 }
 
-bool UHoudiniCSV::GetParticleVectorValueAtTime( const int32& ParticleID, const int32& ColumnIndex, const float& desiredTime, FVector& Vector, const bool& DoSwap, const bool& DoScale )
+bool UHoudiniCSV::GetPointVectorValueAtTime( const int32& PointID, const int32& ColumnIndex, const float& desiredTime, FVector& Vector, const bool& DoSwap, const bool& DoScale )
 {
 	int32 PrevIndex = -1;
 	int32 NextIndex = -1;
 	float PrevWeight = 1.0f;
 
-	if ( !GetParticleLineIndexAtTime( ParticleID, desiredTime, PrevIndex, NextIndex, PrevWeight ) )
+	if ( !GetRowIndexesForPointAtTime( PointID, desiredTime, PrevIndex, NextIndex, PrevWeight ) )
 		return false;
 
 	FVector PrevVector, NextVector;
-	if ( !GetCSVVectorValue( PrevIndex, ColumnIndex, PrevVector, DoSwap, DoScale ) )
+	if ( !GetVectorValue( PrevIndex, ColumnIndex, PrevVector, DoSwap, DoScale ) )
 		return false;
-	if ( !GetCSVVectorValue( NextIndex, ColumnIndex, NextVector, DoSwap, DoScale ) )
+	if ( !GetVectorValue( NextIndex, ColumnIndex, NextVector, DoSwap, DoScale ) )
 		return false;
 
 	Vector = FMath::Lerp(PrevVector, NextVector, PrevWeight);
@@ -915,35 +990,35 @@ bool UHoudiniCSV::GetParticleVectorValueAtTime( const int32& ParticleID, const i
 	return true;
 }
 
-bool UHoudiniCSV::GetParticleLineIndexAtTime(const int32& ParticleID, const float& desiredTime, int32& PrevIndex, int32& NextIndex, float& PrevWeight )
+bool UHoudiniCSV::GetRowIndexesForPointAtTime(const int32& PointID, const float& desiredTime, int32& PrevIndex, int32& NextIndex, float& PrevWeight )
 {
 	float PrevTime = -1.0f;
 	float NextTime = -1.0f;
 
-	// Invalid Particle ID
-	if ( ParticleID < 0 || ParticleID >= NumberOfParticles )
+	// Invalid PointID
+	if ( PointID < 0 || PointID >= NumberOfPoints )
 		return false;
 
 	// Get the spawn time
-	float ParticleSpawnTime = 0.0f;
-	if ( SpawnTimes.IsValidIndex( ParticleID ) )
-		ParticleSpawnTime = SpawnTimes[ ParticleID ];
+	float PointSpawnTime = 0.0f;
+	if ( SpawnTimes.IsValidIndex( PointID ) )
+		PointSpawnTime = SpawnTimes[ PointID ];
 
-	// If particle hasn't spawn before DesiredTime
-	if ( ParticleSpawnTime > desiredTime )
+	// If the point hasn't spawn before DesiredTime
+	if ( PointSpawnTime > desiredTime )
 		return false;
 
 	// Get the life value
-	float ParticleLife = 0.0f;
-	if ( LifeValues.IsValidIndex( ParticleID ) )
-		ParticleLife = LifeValues[ ParticleID ];
+	float PointLife = 0.0f;
+	if ( LifeValues.IsValidIndex( PointID ) )
+		PointLife = LifeValues[ PointID ];
 
-	// If particle is dead before DesiredTime
-	if ( ParticleLife > 0.0f && ( ParticleSpawnTime + ParticleLife < desiredTime ) )	
+	// If the point is dead before DesiredTime
+	if ( PointLife > 0.0f && ( PointSpawnTime + PointLife < desiredTime ) )	
 		return false;
 
 	// We don't have Id information
-	// return ParticleId for the rowIndexes ??
+	// return PointID for the rowIndexes ??
 	if ( IDColumnIndex == INDEX_NONE )
 		return false;
 
@@ -951,10 +1026,10 @@ bool UHoudiniCSV::GetParticleLineIndexAtTime(const int32& ParticleID, const floa
 	if ( TimeColumnIndex == INDEX_NONE )
 		return false;
 
-	// Get the row indexes for this particle
+	// Get the row indexes for this point
 	TArray<int32>* RowIndexes = nullptr;
-	if ( ParticleValueIndexes.IsValidIndex( ParticleID ) )
-		RowIndexes = &( ParticleValueIndexes[ ParticleID ].RowIndexes );
+	if ( PointValueIndexes.IsValidIndex( PointID ) )
+		RowIndexes = &( PointValueIndexes[ PointID ].RowIndexes );
 
 	if ( !RowIndexes )
 		return false;
@@ -963,7 +1038,7 @@ bool UHoudiniCSV::GetParticleLineIndexAtTime(const int32& ParticleID, const floa
 	{
 		// Get the time
 		float currentTime = -1.0f;
-		if ( GetCSVTimeValue( n, currentTime) )
+		if ( GetTimeValue( n, currentTime) )
 		{
 			if ( currentTime == desiredTime )
 			{
@@ -1022,41 +1097,43 @@ bool UHoudiniCSV::GetParticleLineIndexAtTime(const int32& ParticleID, const floa
 // Returns the column index for a given string
 bool UHoudiniCSV::GetColumnIndexFromString( const FString& ColumnTitle, int32& ColumnIndex )
 {
-    if ( !TitleRowArray.Find( ColumnTitle, ColumnIndex ) )
+    if ( !ColumnTitleArray.Find( ColumnTitle, ColumnIndex ) )
 		return true;
 
     // Handle packed positions/normals here
     if ( ColumnTitle.Equals( "P" ) )
-		return TitleRowArray.Find( TEXT( "Px" ), ColumnIndex );
+		return ColumnTitleArray.Find( TEXT( "Px" ), ColumnIndex );
     else if ( ColumnTitle.Equals( "N" ) )
-		return TitleRowArray.Find( TEXT( "Nx" ), ColumnIndex );
+		return ColumnTitleArray.Find( TEXT( "Nx" ), ColumnIndex );
 
     return false;
 }
 
 // Returns the float value at a given point in the CSV file
-bool UHoudiniCSV::GetCSVFloatValue( const int32& lineIndex, const FString& ColumnTitle, float& value )
+bool UHoudiniCSV::GetFloatValue( const int32& rowIndex, const FString& ColumnTitle, float& value )
 {
     int32 ColIndex = -1;
     if ( !GetColumnIndexFromString( ColumnTitle, ColIndex ) )
 		return false;
 
-    return GetCSVFloatValue( lineIndex, ColIndex, value );
+    return GetFloatValue( rowIndex, ColIndex, value );
 }
 
-
+/*
 // Returns the string value at a given point in the CSV file
-bool UHoudiniCSV::GetCSVStringValue( const int32& lineIndex, const FString& ColumnTitle, FString& value )
+bool UHoudiniCSV::GetCSVStringValue( const int32& rowIndex, const FString& ColumnTitle, FString& value )
 {
     int32 ColIndex = -1;
     if ( !GetColumnIndexFromString( ColumnTitle, ColIndex ) )
 		return false;
 
-    return GetCSVStringValue( lineIndex, ColIndex, value );
+    return GetCSVStringValue( rowIndex, ColIndex, value );
 }
+*/
 
 #if WITH_EDITORONLY_DATA
-void UHoudiniCSV::PostInitProperties()
+void
+UHoudiniCSV::PostInitProperties()
 {
     if ( !HasAnyFlags( RF_ClassDefaultObject ) )
     {
@@ -1065,6 +1142,77 @@ void UHoudiniCSV::PostInitProperties()
 
     Super::PostInitProperties();
 }
+
+void
+UHoudiniCSV::PostEditChangeProperty(FPropertyChangedEvent & PropertyChangedEvent)
+{
+	Super::PostEditChangeProperty(PropertyChangedEvent);
+
+	if ( PropertyChangedEvent.GetPropertyName() == TEXT( "SourceTitleRow" ) )
+	{
+		UseCustomTitleRow = true;
+		UpdateFromFile( FileName );
+	}
+	
+}
+
+void
+UHoudiniCSV::GetAssetRegistryTags(TArray< FAssetRegistryTag > & OutTags) const
+{
+	// Add the source filename to the asset thumbnail tooltip
+	OutTags.Add(FAssetRegistryTag("Source FileName", FileName, FAssetRegistryTag::TT_Alphabetical));
+
+	// The Number of rows, columns and point found in the file
+	OutTags.Add(FAssetRegistryTag("Number of Rows", FString::FromInt(NumberOfRows), FAssetRegistryTag::TT_Numerical));
+	OutTags.Add(FAssetRegistryTag("Number of Columns", FString::FromInt(NumberOfColumns), FAssetRegistryTag::TT_Numerical));
+	OutTags.Add(FAssetRegistryTag("Number of Points", FString::FromInt(NumberOfPoints), FAssetRegistryTag::TT_Numerical));
+
+	// The source title row
+	OutTags.Add(FAssetRegistryTag("Original Title Row", SourceTitleRow, FAssetRegistryTag::TT_Alphabetical));
+
+	// The parsed column titles
+	FString ParsedColTtiles;
+	for ( int32 n = 0; n < ColumnTitleArray.Num(); n++ )
+		ParsedColTtiles += TEXT("(") + FString::FromInt( n ) + TEXT(") ") + ColumnTitleArray[ n ] + TEXT(" ");
+
+	OutTags.Add( FAssetRegistryTag( "Parsed Column Titles", ParsedColTtiles, FAssetRegistryTag::TT_Alphabetical ) );
+
+	// And a list of the special attributes we found
+	FString SpecialAttr;
+	if ( IDColumnIndex != INDEX_NONE )
+		SpecialAttr += TEXT("ID ");
+
+	if ( TypeColumnIndex != INDEX_NONE )
+		SpecialAttr += TEXT("Type ");
+
+	if ( PositionColumnIndex != INDEX_NONE )
+		SpecialAttr += TEXT("Position ");
+
+	if ( NormalColumnIndex != INDEX_NONE )
+		SpecialAttr += TEXT("Normal ");
+
+	if ( VelocityColumnIndex != INDEX_NONE )
+		SpecialAttr += TEXT("Velocity ");
+
+	if ( TimeColumnIndex != INDEX_NONE )
+		SpecialAttr += TEXT("Time ");
+
+	if ( ColorColumnIndex != INDEX_NONE )
+		SpecialAttr += TEXT("Color ");
+
+	if ( AlphaColumnIndex != INDEX_NONE )
+		SpecialAttr += TEXT("Alpha ");
+
+	if ( AliveColumnIndex != INDEX_NONE )
+		SpecialAttr += TEXT("Alive ");
+
+	if ( LifeColumnIndex != INDEX_NONE )
+		SpecialAttr += TEXT("Life ");
+
+	OutTags.Add(FAssetRegistryTag("Special Attributes", SpecialAttr, FAssetRegistryTag::TT_Alphabetical));
+
+	Super::GetAssetRegistryTags( OutTags );
+}
 #endif
 
 #undef LOCTEXT_NAMESPACE
diff --git a/Engine/Plugins/FX/HoudiniNiagara/Source/HoudiniNiagara/Private/NiagaraDataInterfaceHoudiniCSV.cpp b/Engine/Plugins/FX/HoudiniNiagara/Source/HoudiniNiagara/Private/NiagaraDataInterfaceHoudiniCSV.cpp
--- a/Engine/Plugins/FX/HoudiniNiagara/Source/HoudiniNiagara/Private/NiagaraDataInterfaceHoudiniCSV.cpp
+++ b/Engine/Plugins/FX/HoudiniNiagara/Source/HoudiniNiagara/Private/NiagaraDataInterfaceHoudiniCSV.cpp
@@ -31,8 +31,8 @@
 UNiagaraDataInterfaceHoudiniCSV::UNiagaraDataInterfaceHoudiniCSV(FObjectInitializer const& ObjectInitializer)
 	: Super(ObjectInitializer)
 {
-    CSVFile = nullptr;
-	LastSpawnedParticleID = -1;
+    HoudiniCSVAsset = nullptr;
+	LastSpawnedPointID = -1;
 }
 
 void UNiagaraDataInterfaceHoudiniCSV::PostInitProperties()
@@ -45,14 +45,14 @@ void UNiagaraDataInterfaceHoudiniCSV::PostInitProperties()
     }
 
     GPUBufferDirty = true;
-	LastSpawnedParticleID = -1;
+	LastSpawnedPointID = -1;
 }
 
 void UNiagaraDataInterfaceHoudiniCSV::PostLoad()
 {
     Super::PostLoad();
     GPUBufferDirty = true;
-	LastSpawnedParticleID = -1;
+	LastSpawnedPointID = -1;
 }
 
 #if WITH_EDITOR
@@ -61,13 +61,13 @@ void UNiagaraDataInterfaceHoudiniCSV::PostEditChangeProperty(struct FPropertyCha
 {
     Super::PostEditChangeProperty(PropertyChangedEvent);
 
-    if (PropertyChangedEvent.Property && PropertyChangedEvent.Property->GetFName() == GET_MEMBER_NAME_CHECKED(UNiagaraDataInterfaceHoudiniCSV, CSVFile))
+    if (PropertyChangedEvent.Property && PropertyChangedEvent.Property->GetFName() == GET_MEMBER_NAME_CHECKED(UNiagaraDataInterfaceHoudiniCSV, HoudiniCSVAsset))
     {
 		Modify();
-		if (CSVFile)
+		if (HoudiniCSVAsset)
 		{
 			GPUBufferDirty = true;
-			LastSpawnedParticleID = -1;
+			LastSpawnedPointID = -1;
 		}
     }
 }
@@ -83,7 +83,7 @@ bool UNiagaraDataInterfaceHoudiniCSV::CopyToInternal(UNiagaraDataInterface* Dest
     if ( !CastedInterface )
 		return false;
 
-    CastedInterface->CSVFile = CSVFile;
+    CastedInterface->HoudiniCSVAsset = HoudiniCSVAsset;
 
     return true;
 }
@@ -95,10 +95,10 @@ bool UNiagaraDataInterfaceHoudiniCSV::Equals(const UNiagaraDataInterface* Other)
 
     const UNiagaraDataInterfaceHoudiniCSV* OtherHNCSV = CastChecked<UNiagaraDataInterfaceHoudiniCSV>(Other);
 
-    if ( OtherHNCSV != nullptr && OtherHNCSV->CSVFile != nullptr && CSVFile )
+    if ( OtherHNCSV != nullptr && OtherHNCSV->HoudiniCSVAsset != nullptr && HoudiniCSVAsset )
     {
 		// Just make sure the two interfaces point to the same file
-		return OtherHNCSV->CSVFile->FileName.Equals( CSVFile->FileName );
+		return OtherHNCSV->HoudiniCSVAsset->FileName.Equals( HoudiniCSVAsset->FileName );
     }
 
     return false;
@@ -108,9 +108,9 @@ bool UNiagaraDataInterfaceHoudiniCSV::Equals(const UNiagaraDataInterface* Other)
 void UNiagaraDataInterfaceHoudiniCSV::GetFunctions(TArray<FNiagaraFunctionSignature>& OutFunctions)
 {
     {
-		// GetCSVFloatValue
+		// GetFloatValue
 		FNiagaraFunctionSignature Sig;
-		Sig.Name = TEXT("GetCSVFloatValue");
+		Sig.Name = TEXT("GetFloatValue");
 		Sig.bMemberFunction = true;
 		Sig.bRequiresContext = false;
 		Sig.Inputs.Add(FNiagaraVariable(FNiagaraTypeDefinition(GetClass()), TEXT("CSV")));			// CSV in
@@ -118,7 +118,7 @@ void UNiagaraDataInterfaceHoudiniCSV::GetFunctions(TArray<FNiagaraFunctionSignat
 		Sig.Inputs.Add(FNiagaraVariable(FNiagaraTypeDefinition::GetIntDef(), TEXT("Col")));			// Col Index In
 		Sig.Outputs.Add(FNiagaraVariable(FNiagaraTypeDefinition::GetFloatDef(), TEXT("Value")));	// Float Out
 
-		Sig.SetDescription( LOCTEXT( "DataInterfaceHoudini_GetCSVFloatValue",
+		Sig.SetDescription( LOCTEXT( "DataInterfaceHoudini_GetFloatValue",
 			"Returns the float value in the CSV file for a given Row and Column.\n" ) );
 
 		OutFunctions.Add( Sig );
@@ -126,9 +126,9 @@ void UNiagaraDataInterfaceHoudiniCSV::GetFunctions(TArray<FNiagaraFunctionSignat
 
 	/*
     {
-		// GetCSVFloatValueByString
+		// GetFloatValueByString
 		FNiagaraFunctionSignature Sig;
-		Sig.Name = TEXT("GetCSVFloatValueByString");
+		Sig.Name = TEXT("GetFloatValueByString");
 		Sig.bMemberFunction = true;
 		Sig.bRequiresContext = false;
 		Sig.Inputs.Add(FNiagaraVariable(FNiagaraTypeDefinition(GetClass()), TEXT("CSV")));			// CSV in
@@ -136,7 +136,7 @@ void UNiagaraDataInterfaceHoudiniCSV::GetFunctions(TArray<FNiagaraFunctionSignat
 		Sig.Inputs.Add(FNiagaraVariable(FNiagaraTypeDefinition::GetStringDef(), TEXT("ColTitle")));	// Col Title In
 		Sig.Outputs.Add(FNiagaraVariable(FNiagaraTypeDefinition::GetFloatDef(), TEXT("Value")));	// Float Out
 
-		Sig.SetDescription( LOCTEXT( "DataInterfaceHoudini_GetCSVFloatValueByString",
+		Sig.SetDescription( LOCTEXT( "DataInterfaceHoudini_GetFloatValueByString",
 		"Returns the float value in the CSV file for a given Row, in the column corresponding to the ColTitle string.\n" ) );
 
 		OutFunctions.Add(Sig);
@@ -144,9 +144,9 @@ void UNiagaraDataInterfaceHoudiniCSV::GetFunctions(TArray<FNiagaraFunctionSignat
 	*/
 
 	{
-		// GetCSVVectorValue
+		// GetVectorValue
 		FNiagaraFunctionSignature Sig;
-		Sig.Name = TEXT("GetCSVVectorValue");
+		Sig.Name = TEXT("GetVectorValue");
 		Sig.bMemberFunction = true;
 		Sig.bRequiresContext = false;
 		Sig.Inputs.Add(FNiagaraVariable(FNiagaraTypeDefinition(GetClass()), TEXT("CSV")));			// CSV in
@@ -154,32 +154,51 @@ void UNiagaraDataInterfaceHoudiniCSV::GetFunctions(TArray<FNiagaraFunctionSignat
 		Sig.Inputs.Add(FNiagaraVariable(FNiagaraTypeDefinition::GetIntDef(), TEXT("Col")));			// Col Index In
 		Sig.Outputs.Add(FNiagaraVariable(FNiagaraTypeDefinition::GetVec3Def(), TEXT("Value")));		// Vector3 Out
 
-		Sig.SetDescription( LOCTEXT( "DataInterfaceHoudini_GetCSVVectorValue",
+		Sig.SetDescription( LOCTEXT( "DataInterfaceHoudini_GetVectorValue",
 			"Returns a Vector3 in the CSV file for a given Row and Column.\nThe returned Vector is converted from Houdini's coordinate system to Unreal's." ) );
 
 		OutFunctions.Add(Sig);
 	}
 
+	{
+		// GetVectorValueEx
+		FNiagaraFunctionSignature Sig;
+		Sig.Name = TEXT("GetVectorValueEx");
+		Sig.bMemberFunction = true;
+		Sig.bRequiresContext = false;
+		Sig.Inputs.Add(FNiagaraVariable(FNiagaraTypeDefinition(GetClass()), TEXT("CSV")));			// CSV in
+		Sig.Inputs.Add(FNiagaraVariable(FNiagaraTypeDefinition::GetIntDef(), TEXT("Row")));			// Row Index In
+		Sig.Inputs.Add(FNiagaraVariable(FNiagaraTypeDefinition::GetIntDef(), TEXT("Col")));			// Col Index In
+		Sig.Inputs.Add(FNiagaraVariable(FNiagaraTypeDefinition::GetBoolDef(), TEXT("DoSwap")));		// DoSwap in
+		Sig.Inputs.Add(FNiagaraVariable(FNiagaraTypeDefinition::GetBoolDef(), TEXT("DoScale")));	// DoScale in
+		Sig.Outputs.Add(FNiagaraVariable(FNiagaraTypeDefinition::GetVec3Def(), TEXT("Value")));		// Vector3 Out
+
+		Sig.SetDescription(LOCTEXT("DataInterfaceHoudini_GetVectorValueEx",
+			"Returns a Vector3 in the CSV file for a given Row and Column.\nThe DoSwap parameter indicates if the vector should be converted from Houdini*s coordinate system to Unreal's.\nThe DoScale parameter decides if the Vector value should be converted from meters (Houdini) to centimeters (Unreal)."));
+
+		OutFunctions.Add(Sig);
+	}
+
     {
-		// GetCSVPosition
+		// GetPosition
 		FNiagaraFunctionSignature Sig;
-		Sig.Name = TEXT("GetCSVPosition");
+		Sig.Name = TEXT("GetPosition");
 		Sig.bMemberFunction = true;
 		Sig.bRequiresContext = false;
 		Sig.Inputs.Add(FNiagaraVariable(FNiagaraTypeDefinition(GetClass()), TEXT("CSV")));			// CSV in
 		Sig.Inputs.Add(FNiagaraVariable(FNiagaraTypeDefinition::GetIntDef(), TEXT("Row")));			// Row Index In
 		Sig.Outputs.Add(FNiagaraVariable(FNiagaraTypeDefinition::GetVec3Def(), TEXT("Position")));	// Vector3 Out
 
-		Sig.SetDescription( LOCTEXT( "DataInterfaceHoudini_GetCSVPosition",
+		Sig.SetDescription( LOCTEXT( "DataInterfaceHoudini_GetPosition",
 			"Helper function returning the position value for a given Row in the CSV file.\nThe returned Position vector is converted from Houdini's coordinate system to Unreal's." ) );
 
 		OutFunctions.Add(Sig);
     }
 
     {
-		// GetCSVPositionAndTime
+		// GetPositionAndTime
 		FNiagaraFunctionSignature Sig;
-		Sig.Name = TEXT("GetCSVPositionAndTime");
+		Sig.Name = TEXT("GetPositionAndTime");
 		Sig.bMemberFunction = true;
 		Sig.bRequiresContext = false;
 		Sig.Inputs.Add(FNiagaraVariable(FNiagaraTypeDefinition(GetClass()), TEXT("CSV")));			// CSV in
@@ -187,116 +206,116 @@ void UNiagaraDataInterfaceHoudiniCSV::GetFunctions(TArray<FNiagaraFunctionSignat
 		Sig.Outputs.Add(FNiagaraVariable(FNiagaraTypeDefinition::GetVec3Def(), TEXT("Position")));	// Vector3 Out
 		Sig.Outputs.Add(FNiagaraVariable(FNiagaraTypeDefinition::GetFloatDef(), TEXT("Time")));		// float Out
 
-		Sig.SetDescription( LOCTEXT( "DataInterfaceHoudini_GetCSVPositionAndTime",
+		Sig.SetDescription( LOCTEXT( "DataInterfaceHoudini_GetPositionAndTime",
 			"Helper function returning the position and time values for a given Row in the CSV file.\nThe returned Position vector is converted from Houdini's coordinate system to Unreal's." ) );
 
 		OutFunctions.Add(Sig);
     }
 
     {
-		// GetCSVNormal
+		// GetNormal
 		FNiagaraFunctionSignature Sig;
-		Sig.Name = TEXT("GetCSVNormal");
+		Sig.Name = TEXT("GetNormal");
 		Sig.bMemberFunction = true;
 		Sig.bRequiresContext = false;
 		Sig.Inputs.Add(FNiagaraVariable(FNiagaraTypeDefinition(GetClass()), TEXT("CSV")));			// CSV in
 		Sig.Inputs.Add(FNiagaraVariable(FNiagaraTypeDefinition::GetIntDef(), TEXT("Row")));			// Row Index In
 		Sig.Outputs.Add(FNiagaraVariable(FNiagaraTypeDefinition::GetVec3Def(), TEXT("Normal")));	// Vector3 Out
 
-		Sig.SetDescription( LOCTEXT( "DataInterfaceHoudini_GetCSVNormal",
+		Sig.SetDescription( LOCTEXT( "DataInterfaceHoudini_GetNormal",
 			"Helper function returning the normal value for a given Row in the CSV file.\nThe returned Normal vector is converted from Houdini's coordinate system to Unreal's." ) );
 
 		OutFunctions.Add(Sig);
     }
 
     {
-		// GetCSVTime
+		// GetTime
 		FNiagaraFunctionSignature Sig;
-		Sig.Name = TEXT("GetCSVTime");
+		Sig.Name = TEXT("GetTime");
 		Sig.bMemberFunction = true;
 		Sig.bRequiresContext = false;
 		Sig.Inputs.Add(FNiagaraVariable(FNiagaraTypeDefinition(GetClass()), TEXT("CSV")));			// CSV in
 		Sig.Inputs.Add(FNiagaraVariable(FNiagaraTypeDefinition::GetIntDef(), TEXT("Row")));			// Row Index In
 		Sig.Outputs.Add(FNiagaraVariable(FNiagaraTypeDefinition::GetFloatDef(), TEXT("Time")));		// Float Out
 
-		Sig.SetDescription( LOCTEXT("DataInterfaceHoudini_GetCSVTime",
+		Sig.SetDescription( LOCTEXT("DataInterfaceHoudini_GetTime",
 			"Helper function returning the time value for a given Row in the CSV file.\n") );
 
 		OutFunctions.Add(Sig);
     }
 
 	{
-		// GetCSVVelocity
+		// GetVelocity
 		FNiagaraFunctionSignature Sig;
-		Sig.Name = TEXT("GetCSVVelocity");
+		Sig.Name = TEXT("GetVelocity");
 		Sig.bMemberFunction = true;
 		Sig.bRequiresContext = false;
 		Sig.Inputs.Add(FNiagaraVariable(FNiagaraTypeDefinition(GetClass()), TEXT("CSV")));			// CSV in
 		Sig.Inputs.Add(FNiagaraVariable(FNiagaraTypeDefinition::GetIntDef(), TEXT("Row")));			// Row Index In
 		Sig.Outputs.Add(FNiagaraVariable(FNiagaraTypeDefinition::GetVec3Def(), TEXT("Velocity")));	// Vector3 Out
 
-		Sig.SetDescription(LOCTEXT("DataInterfaceHoudini_GetCSVVelocity",
+		Sig.SetDescription(LOCTEXT("DataInterfaceHoudini_GetVelocity",
 			"Helper function returning the velocity value for a given Row in the CSV file.\nThe returned velocity vector is converted from Houdini's coordinate system to Unreal's."));
 
 		OutFunctions.Add(Sig);
 	}
 
 	{
-		// GetCSVColor
+		// GetColor
 		FNiagaraFunctionSignature Sig;
-		Sig.Name = TEXT("GetCSVColor");
+		Sig.Name = TEXT("GetColor");
 		Sig.bMemberFunction = true;
 		Sig.bRequiresContext = false;
 		Sig.Inputs.Add(FNiagaraVariable(FNiagaraTypeDefinition(GetClass()), TEXT("CSV")));			// CSV in
 		Sig.Inputs.Add(FNiagaraVariable(FNiagaraTypeDefinition::GetIntDef(), TEXT("Row")));			// Row Index In
 		Sig.Outputs.Add(FNiagaraVariable(FNiagaraTypeDefinition::GetColorDef(), TEXT("Color")));	// Color Out
 
-		Sig.SetDescription(LOCTEXT("DataInterfaceHoudini_GetCSVColor",
+		Sig.SetDescription(LOCTEXT("DataInterfaceHoudini_GetColor",
 			"Helper function returning the color value for a given Row in the CSV file."));
 
 		OutFunctions.Add(Sig);
 	}
 
     {
-		// GetNumberOfParticlesInCSV
+		// GetNumberOfPoints
 		FNiagaraFunctionSignature Sig;
-		Sig.Name = TEXT("GetNumberOfParticlesInCSV");
+		Sig.Name = TEXT("GetNumberOfPoints");
 		Sig.bMemberFunction = true;
 		Sig.bRequiresContext = false;
 		Sig.Inputs.Add(FNiagaraVariable(FNiagaraTypeDefinition(GetClass()), TEXT("CSV")));					// CSV in
-		Sig.Outputs.Add(FNiagaraVariable(FNiagaraTypeDefinition::GetIntDef(), TEXT("NumberOfParticles")));  // Int Out
+		Sig.Outputs.Add(FNiagaraVariable(FNiagaraTypeDefinition::GetIntDef(), TEXT("NumberOfPoints")));  // Int Out
 
-		Sig.SetDescription( LOCTEXT( "DataInterfaceHoudini_GetNumberOfParticlesInCSV",
-			"Returns the number of particles (with different id values) in the CSV file.\n" ) );
+		Sig.SetDescription( LOCTEXT( "DataInterfaceHoudini_GetNumberOfPoints",
+			"Returns the number of points (with different id values) in the CSV file.\n" ) );
 
 		OutFunctions.Add(Sig);
     }
 
 	{
-		// GetNumberOfRowsInCSV
+		// GetNumberOfRows
 		FNiagaraFunctionSignature Sig;
-		Sig.Name = TEXT("GetNumberOfRowsInCSV");
+		Sig.Name = TEXT("GetNumberOfRows");
 		Sig.bMemberFunction = true;
 		Sig.bRequiresContext = false;
 		Sig.Inputs.Add( FNiagaraVariable( FNiagaraTypeDefinition( GetClass()), TEXT("CSV") ) );					// CSV in
 		Sig.Outputs.Add( FNiagaraVariable( FNiagaraTypeDefinition::GetIntDef(), TEXT("NumberOfRows") ) );		// Int Out
 		
-		Sig.SetDescription( LOCTEXT( "DataInterfaceHoudini_GetNumberOfRowsInCSV",
+		Sig.SetDescription( LOCTEXT( "DataInterfaceHoudini_GetNumberOfRows",
 			"Returns the number of rows in the CSV file.\nOnly the number of value rows is returned, the first \"Title\" Row is ignored." ) );
 
 		OutFunctions.Add(Sig);
 	}
 
-		{
-		// GetNumberOfColumnsInCSV
+	{
+		// GetNumberOfColumns
 		FNiagaraFunctionSignature Sig;
-		Sig.Name = TEXT("GetNumberOfColumnsInCSV");
+		Sig.Name = TEXT("GetNumberOfColumns");
 		Sig.bMemberFunction = true;
 		Sig.bRequiresContext = false;
 		Sig.Inputs.Add( FNiagaraVariable( FNiagaraTypeDefinition( GetClass()), TEXT("CSV") ) );						// CSV in
 		Sig.Outputs.Add( FNiagaraVariable( FNiagaraTypeDefinition::GetIntDef(), TEXT("NumberOfColumns") ) );		// Int Out
 		
-		Sig.SetDescription( LOCTEXT( "DataInterfaceHoudini_GetNumberOfColumnsInCSV",
+		Sig.SetDescription( LOCTEXT( "DataInterfaceHoudini_GetNumberOfColumns",
 			"Returns the number of columns in the CSV file." ) );
 
 		OutFunctions.Add(Sig);
@@ -319,9 +338,9 @@ void UNiagaraDataInterfaceHoudiniCSV::GetFunctions(TArray<FNiagaraFunctionSignat
     }
 
     {
-		// GetParticleIndexesToSpawnAtTime
+		// GetPointIdsToSpawnAtTime
 		FNiagaraFunctionSignature Sig;
-		Sig.Name = TEXT("GetParticleIndexesToSpawnAtTime");
+		Sig.Name = TEXT("GetPointIDsToSpawnAtTime");
 		Sig.bMemberFunction = true;
 		Sig.bRequiresContext = false;
 		Sig.Inputs.Add(FNiagaraVariable( FNiagaraTypeDefinition(GetClass()), TEXT("CSV") ) );				// CSV in
@@ -330,97 +349,151 @@ void UNiagaraDataInterfaceHoudiniCSV::GetFunctions(TArray<FNiagaraFunctionSignat
 		Sig.Outputs.Add(FNiagaraVariable( FNiagaraTypeDefinition::GetIntDef(), TEXT("MaxIndex") ) );	    // Int Out
 		Sig.Outputs.Add(FNiagaraVariable( FNiagaraTypeDefinition::GetIntDef(), TEXT("Count") ) );		    // Int Out
 
-		Sig.SetDescription( LOCTEXT( "DataInterfaceHoudini_GetParticleIndexesToSpawnAtTime",
-			"Returns the count and particle IDs of the particles that should spawn at a given time value." ) );
+		Sig.SetDescription( LOCTEXT( "DataInterfaceHoudini_GetPointIDsToSpawnAtTime",
+			"Returns the count and point IDs of the points that should spawn for a given time value." ) );
 
 		OutFunctions.Add(Sig);
     }
 
 	{
-		// GetRowIndexesForParticleAtTime
+		// GetRowIndexesForPointAtTime
 		FNiagaraFunctionSignature Sig;
-		Sig.Name = TEXT("GetRowIndexesForParticleAtTime");
+		Sig.Name = TEXT("GetRowIndexesForPointAtTime");
 		Sig.bMemberFunction = true;
 		Sig.bRequiresContext = false;
 		Sig.Inputs.Add(FNiagaraVariable(FNiagaraTypeDefinition(GetClass()), TEXT("CSV")));					// CSV in
-		Sig.Inputs.Add(FNiagaraVariable(FNiagaraTypeDefinition::GetIntDef(), TEXT("N")));					// Point Number In
+		Sig.Inputs.Add(FNiagaraVariable(FNiagaraTypeDefinition::GetIntDef(), TEXT("PointID")));				// Point Number In
 		Sig.Inputs.Add(FNiagaraVariable(FNiagaraTypeDefinition::GetFloatDef(), TEXT("Time")));				// Time in
 		Sig.Outputs.Add(FNiagaraVariable(FNiagaraTypeDefinition::GetIntDef(), TEXT("PreviousRow")));		// Int Out
 		Sig.Outputs.Add(FNiagaraVariable(FNiagaraTypeDefinition::GetIntDef(), TEXT("NextRow")));			// Int Out
 		Sig.Outputs.Add(FNiagaraVariable(FNiagaraTypeDefinition::GetFloatDef(), TEXT("PrevWeight")));		// Float Out
 
-		Sig.SetDescription( LOCTEXT( "DataInterfaceHoudini_GetRowIndexesForParticleAtTime",
-			"Returns the row indexes for a given particle at a given time.\nThe previous row, next row and weight can then be used to Lerp between values.") );
+		Sig.SetDescription( LOCTEXT( "DataInterfaceHoudini_GetRowIndexesForPointAtTime",
+			"Returns the row indexes for a given point at a given time.\nThe previous row, next row and weight can then be used to Lerp between values.") );
 
 		OutFunctions.Add(Sig);
 	}
 
 	{
-		// GetParticlePositionAtTime
+		// GetPointPositionAtTime
 		FNiagaraFunctionSignature Sig;
-		Sig.Name = TEXT("GetParticlePositionAtTime");
+		Sig.Name = TEXT("GetPointPositionAtTime");
 		Sig.bMemberFunction = true;
 		Sig.bRequiresContext = false;
 		Sig.Inputs.Add(FNiagaraVariable(FNiagaraTypeDefinition(GetClass()), TEXT("CSV")));				// CSV in
-		Sig.Inputs.Add(FNiagaraVariable(FNiagaraTypeDefinition::GetIntDef(), TEXT("N")));				// Point Number In
+		Sig.Inputs.Add(FNiagaraVariable(FNiagaraTypeDefinition::GetIntDef(), TEXT("PointID")));				// Point Number In
 		Sig.Inputs.Add(FNiagaraVariable(FNiagaraTypeDefinition::GetFloatDef(), TEXT("Time")));		    // Time in
 		Sig.Outputs.Add(FNiagaraVariable(FNiagaraTypeDefinition::GetVec3Def(), TEXT("Position")));		// Vector3 Out
 
-		Sig.SetDescription( LOCTEXT( "DataInterfaceHoudini_GetParticlePositionAtTime",
-			"Helper function returning the linearly interpolated position for a given particle at a given time.\nThe returned Position vector is converted from Houdini's coordinate system to Unreal's.") );
+		Sig.SetDescription( LOCTEXT( "DataInterfaceHoudini_GetPointPositionAtTime",
+			"Helper function returning the linearly interpolated position for a given point at a given time.\nThe returned Position vector is converted from Houdini's coordinate system to Unreal's.") );
 
 		OutFunctions.Add(Sig);
 	}
 
 	{
-		// GetParticleValueAtTime
+		// GetPointValueAtTime
 		FNiagaraFunctionSignature Sig;
-		Sig.Name = TEXT("GetParticleValueAtTime");
+		Sig.Name = TEXT("GetPointValueAtTime");
 		Sig.bMemberFunction = true;
 		Sig.bRequiresContext = false;
 		Sig.Inputs.Add(FNiagaraVariable(FNiagaraTypeDefinition(GetClass()), TEXT("CSV")));				// CSV in
-		Sig.Inputs.Add(FNiagaraVariable(FNiagaraTypeDefinition::GetIntDef(), TEXT("N")));				// Point Number In		
+		Sig.Inputs.Add(FNiagaraVariable(FNiagaraTypeDefinition::GetIntDef(), TEXT("PointID")));				// Point Number In		
 		Sig.Inputs.Add(FNiagaraVariable(FNiagaraTypeDefinition::GetIntDef(), TEXT("Col")));				// Col Index In
 		Sig.Inputs.Add(FNiagaraVariable(FNiagaraTypeDefinition::GetFloatDef(), TEXT("Time")));		    // Time in		
 		Sig.Outputs.Add(FNiagaraVariable(FNiagaraTypeDefinition::GetFloatDef(), TEXT("Value")));		// Float Out
 
-		Sig.SetDescription( LOCTEXT( "DataInterfaceHoudini_GetParticleValueAtTime",
-			"Returns the linearly interpolated value in the specified column for a given particle at a given time." ) );
+		Sig.SetDescription( LOCTEXT( "DataInterfaceHoudini_GetPointValueAtTime",
+			"Returns the linearly interpolated value in the specified column for a given point at a given time." ) );
 
 		OutFunctions.Add(Sig);
 	}
 
 	{
-		// GetParticleVectorValueAtTime
+		// GetPointVectorValueAtTime
 		FNiagaraFunctionSignature Sig;
-		Sig.Name = TEXT("GetParticleVectorValueAtTime");
+		Sig.Name = TEXT("GetPointVectorValueAtTime");
 		Sig.bMemberFunction = true;
 		Sig.bRequiresContext = false;
 		Sig.Inputs.Add(FNiagaraVariable(FNiagaraTypeDefinition(GetClass()), TEXT("CSV")));				// CSV in
-		Sig.Inputs.Add(FNiagaraVariable(FNiagaraTypeDefinition::GetIntDef(), TEXT("N")));				// Point Number In
+		Sig.Inputs.Add(FNiagaraVariable(FNiagaraTypeDefinition::GetIntDef(), TEXT("PointID")));			// Point ID In
 		Sig.Inputs.Add(FNiagaraVariable(FNiagaraTypeDefinition::GetIntDef(), TEXT("Col")));				// Col Index In
 		Sig.Inputs.Add(FNiagaraVariable(FNiagaraTypeDefinition::GetFloatDef(), TEXT("Time")));		    // Time in		
 		Sig.Outputs.Add(FNiagaraVariable(FNiagaraTypeDefinition::GetVec3Def(), TEXT("Value")));			// Vector3 Out
 
-		Sig.SetDescription( LOCTEXT( "DataInterfaceHoudini_GetParticleVectorValueAtTime",
-			"Helper function returning the linearly interpolated Vector value in the specified column for a given particle at a given time.\nThe returned Vector is converted from Houdini's coordinate system to Unreal's." ) );
+		Sig.SetDescription( LOCTEXT( "DataInterfaceHoudini_GetPointVectorValueAtTime",
+			"Helper function returning the linearly interpolated Vector value in the specified column for a given point at a given time.\nThe returned Vector is converted from Houdini's coordinate system to Unreal's." ) );
 
 		OutFunctions.Add(Sig);
 	}
 
 	{
-		// GetParticleLifeAtTime
+		// GetPointVectorValueAtTimeEx
 		FNiagaraFunctionSignature Sig;
-		Sig.Name = TEXT("GetParticleLifeAtTime");
+		Sig.Name = TEXT("GetPointVectorValueAtTimeEx");
 		Sig.bMemberFunction = true;
 		Sig.bRequiresContext = false;
 		Sig.Inputs.Add(FNiagaraVariable(FNiagaraTypeDefinition(GetClass()), TEXT("CSV")));				// CSV in
-		Sig.Inputs.Add(FNiagaraVariable(FNiagaraTypeDefinition::GetIntDef(), TEXT("N")));				// Point Number In		
+		Sig.Inputs.Add(FNiagaraVariable(FNiagaraTypeDefinition::GetIntDef(), TEXT("PointID")));			// Point ID In
+		Sig.Inputs.Add(FNiagaraVariable(FNiagaraTypeDefinition::GetIntDef(), TEXT("Col")));				// Col Index In
+		Sig.Inputs.Add(FNiagaraVariable(FNiagaraTypeDefinition::GetFloatDef(), TEXT("Time")));		    // Time in
+		Sig.Inputs.Add(FNiagaraVariable(FNiagaraTypeDefinition::GetBoolDef(), TEXT("DoSwap")));			// DoSwap in
+		Sig.Inputs.Add(FNiagaraVariable(FNiagaraTypeDefinition::GetBoolDef(), TEXT("DoScale")));		// DoScale in
+		Sig.Outputs.Add(FNiagaraVariable(FNiagaraTypeDefinition::GetVec3Def(), TEXT("Value")));			// Vector3 Out
+
+		Sig.SetDescription(LOCTEXT("DataInterfaceHoudini_GetPointVectorValueAtTime",
+			"Helper function returning the linearly interpolated Vector value in the specified column for a given point at a given time.\nThe DoSwap parameter indicates if the vector should be converted from Houdini*s coordinate system to Unreal's.\nThe DoScale parameter decides if the Vector value should be converted from meters (Houdini) to centimeters (Unreal)." ) );
+
+		OutFunctions.Add(Sig);
+	}
+
+	{
+		// GetPointLife
+		FNiagaraFunctionSignature Sig;
+		Sig.Name = TEXT("GetPointLife");
+		Sig.bMemberFunction = true;
+		Sig.bRequiresContext = false;
+		Sig.Inputs.Add(FNiagaraVariable(FNiagaraTypeDefinition(GetClass()), TEXT("CSV")));				// CSV in
+		Sig.Inputs.Add(FNiagaraVariable(FNiagaraTypeDefinition::GetIntDef(), TEXT("PointID")));			// Point Number In		
+		Sig.Outputs.Add(FNiagaraVariable(FNiagaraTypeDefinition::GetFloatDef(), TEXT("Life")));			// Float Out
+
+		Sig.SetDescription(LOCTEXT("DataInterfaceHoudini_GetPointLife",
+			"Helper function returning the life value for a given point when spawned.\nThe life value is either calculated from the alive attribute or is the life attribute at spawn time."));
+
+		OutFunctions.Add(Sig);
+	}
+
+	/*
+	{
+		// GetPointLifeAtTime
+		FNiagaraFunctionSignature Sig;
+		Sig.Name = TEXT("GetPointLifeAtTime");
+		Sig.bMemberFunction = true;
+		Sig.bRequiresContext = false;
+		Sig.Inputs.Add(FNiagaraVariable(FNiagaraTypeDefinition(GetClass()), TEXT("CSV")));				// CSV in
+		Sig.Inputs.Add(FNiagaraVariable(FNiagaraTypeDefinition::GetIntDef(), TEXT("PointID")));			// Point Number In		
 		Sig.Inputs.Add(FNiagaraVariable(FNiagaraTypeDefinition::GetFloatDef(), TEXT("Time")));		    // Time in		
-		Sig.Outputs.Add(FNiagaraVariable(FNiagaraTypeDefinition::GetFloatDef(), TEXT("Value")));		// Float Out
+		Sig.Outputs.Add(FNiagaraVariable(FNiagaraTypeDefinition::GetFloatDef(), TEXT("Life")));			// Float Out
+
+		Sig.SetDescription( LOCTEXT( "DataInterfaceHoudini_GetPointLifeAtTime",
+			"Helper function returning the remaining life for a given point in the CSV file at a given time." ) );
+
+		OutFunctions.Add(Sig);
+	}
+	*/
 
-		Sig.SetDescription( LOCTEXT( "DataInterfaceHoudini_GetParticleLifeAtTime",
-			"Helper function returning the remaining life for a given particle in the CSV file at a given time." ) );
+	{
+		// GetPointType
+		FNiagaraFunctionSignature Sig;
+		Sig.Name = TEXT("GetPointType");
+		Sig.bMemberFunction = true;
+		Sig.bRequiresContext = false;
+		Sig.Inputs.Add(FNiagaraVariable(FNiagaraTypeDefinition(GetClass()), TEXT("CSV")));				// CSV in
+		Sig.Inputs.Add(FNiagaraVariable(FNiagaraTypeDefinition::GetIntDef(), TEXT("PointID")));			// Point Number In		
+		Sig.Outputs.Add(FNiagaraVariable(FNiagaraTypeDefinition::GetIntDef(), TEXT("Type")));			// Int Out
+
+		Sig.SetDescription(LOCTEXT("DataInterfaceHoudini_GetPointType",
+			"Helper function returning the type value for a given point when spawned.\n"));
 
 		OutFunctions.Add(Sig);
 	}
@@ -436,14 +509,14 @@ bool UNiagaraDataInterfaceHoudiniCSV::GetFunctionHLSL(const FName& DefinitionFun
     //FString SecondBufferName = Descriptors[1].BufferParamName;
 
     /*
-    if (InstanceFunctionName.Contains( TEXT( "GetCSVFloatValue") ) )
+    if (InstanceFunctionName.Contains( TEXT( "GetFloatValue") ) )
     {
 	FString BufferName = Descriptors[0].BufferParamName;
 	OutHLSL += TEXT("void ") + InstanceFunctionName + TEXT("(in float In_Row, in float In_Col, out float Out_Value) \n{\n");
 	OutHLSL += TEXT("\t Out_Value = ") + BufferName + TEXT("[(int)(In_Row + ( In_Col * ") + FString::FromInt(NumberOfRows) + TEXT(") ) ];");
 	OutHLSL += TEXT("\n}\n");
     }
-    else if (InstanceFunctionName.Contains(TEXT("GetCSVPosition")))
+    else if (InstanceFunctionName.Contains(TEXT("GetPosition")))
     {
 	FString BufferName = Descriptors[1].BufferParamName;
 	OutHLSL += TEXT("void ") + InstanceFunctionName + TEXT("(in float In_N, out float3 Out_Value) \n{\n");
@@ -452,7 +525,7 @@ bool UNiagaraDataInterfaceHoudiniCSV::GetFunctionHLSL(const FName& DefinitionFun
 	OutHLSL += TEXT("\t Out_Value.z = ") + BufferName + TEXT("[(int)(In_N + ( 2 * ") + FString::FromInt(NumberOfRows) + TEXT(") ) ];");
 	OutHLSL += TEXT("\n}\n");
     }
-    else if (InstanceFunctionName.Contains(TEXT("GetCSVNormal")))
+    else if (InstanceFunctionName.Contains(TEXT("GetNormal")))
     {
 	FString BufferName = Descriptors[2].BufferParamName;
 	OutHLSL += TEXT("void ") + InstanceFunctionName + TEXT("(in float In_N, out float3 Out_Value) \n{\n");
@@ -461,7 +534,7 @@ bool UNiagaraDataInterfaceHoudiniCSV::GetFunctionHLSL(const FName& DefinitionFun
 	OutHLSL += TEXT("\t Out_Value.z = ") + BufferName + TEXT("[(int)(In_N + ( 2 * ") + FString::FromInt(NumberOfRows) + TEXT(") ) ];");
 	OutHLSL += TEXT("\n}\n");
     }
-    else if (InstanceFunctionName.Contains(TEXT("GetCSVTime")))
+    else if (InstanceFunctionName.Contains(TEXT("GetTime")))
     {
 	FString BufferName = Descriptors[3].BufferParamName;
 	OutHLSL += TEXT("void ") + InstanceFunctionName + TEXT("(in float In_N, out float Out_Value) \n{\n");
@@ -474,7 +547,7 @@ bool UNiagaraDataInterfaceHoudiniCSV::GetFunctionHLSL(const FName& DefinitionFun
 	OutHLSL += TEXT("\t Out_Value = ") + FString::FromInt(NumberOfRows) + TEXT(";");
 	OutHLSL += TEXT("\n}\n");
     }*/
-    /*else if (InstanceFunctionName.Contains(TEXT("GetCSVPositionAndTime")))
+    /*else if (InstanceFunctionName.Contains(TEXT("GetPositionAndTime")))
     {
 	OutHLSL += TEXT("void ") + FunctionName + TEXT("(in float In_N, out float4 Out_Value) \n{\n");
 	OutHLSL += TEXT("\t Out_Value.x = ") + BufferName + TEXT("[(int)(In_N) ];");
@@ -509,99 +582,121 @@ void UNiagaraDataInterfaceHoudiniCSV::GetParameterDefinitionHLSL(FNiagaraDataInt
     OutHLSL += TEXT("Buffer<float> ") + BufferName + TEXT(";\n");
 }
 
-DEFINE_NDI_RAW_FUNC_BINDER(UNiagaraDataInterfaceHoudiniCSV, GetCSVFloatValue);
-//DEFINE_NDI_RAW_FUNC_BINDER(UNiagaraDataInterfaceHoudiniCSV, GetCSVFloatValueByString);
-DEFINE_NDI_RAW_FUNC_BINDER(UNiagaraDataInterfaceHoudiniCSV, GetCSVVectorValue);
-DEFINE_NDI_RAW_FUNC_BINDER(UNiagaraDataInterfaceHoudiniCSV, GetCSVPosition);
-DEFINE_NDI_RAW_FUNC_BINDER(UNiagaraDataInterfaceHoudiniCSV, GetCSVNormal);
-DEFINE_NDI_RAW_FUNC_BINDER(UNiagaraDataInterfaceHoudiniCSV, GetCSVTime);
-DEFINE_NDI_RAW_FUNC_BINDER(UNiagaraDataInterfaceHoudiniCSV, GetCSVColor);
-DEFINE_NDI_RAW_FUNC_BINDER(UNiagaraDataInterfaceHoudiniCSV, GetCSVVelocity);
-DEFINE_NDI_RAW_FUNC_BINDER(UNiagaraDataInterfaceHoudiniCSV, GetCSVPositionAndTime);
+DEFINE_NDI_RAW_FUNC_BINDER(UNiagaraDataInterfaceHoudiniCSV, GetFloatValue);
+//DEFINE_NDI_RAW_FUNC_BINDER(UNiagaraDataInterfaceHoudiniCSV, GetFloatValueByString);
+DEFINE_NDI_RAW_FUNC_BINDER(UNiagaraDataInterfaceHoudiniCSV, GetVectorValue);
+DEFINE_NDI_RAW_FUNC_BINDER(UNiagaraDataInterfaceHoudiniCSV, GetVectorValueEx);
+DEFINE_NDI_RAW_FUNC_BINDER(UNiagaraDataInterfaceHoudiniCSV, GetPosition);
+DEFINE_NDI_RAW_FUNC_BINDER(UNiagaraDataInterfaceHoudiniCSV, GetNormal);
+DEFINE_NDI_RAW_FUNC_BINDER(UNiagaraDataInterfaceHoudiniCSV, GetTime);
+DEFINE_NDI_RAW_FUNC_BINDER(UNiagaraDataInterfaceHoudiniCSV, GetColor);
+DEFINE_NDI_RAW_FUNC_BINDER(UNiagaraDataInterfaceHoudiniCSV, GetVelocity);
+DEFINE_NDI_RAW_FUNC_BINDER(UNiagaraDataInterfaceHoudiniCSV, GetPositionAndTime);
 DEFINE_NDI_RAW_FUNC_BINDER(UNiagaraDataInterfaceHoudiniCSV, GetLastRowIndexAtTime);
-DEFINE_NDI_RAW_FUNC_BINDER(UNiagaraDataInterfaceHoudiniCSV, GetParticleIndexesToSpawnAtTime);
-DEFINE_NDI_RAW_FUNC_BINDER(UNiagaraDataInterfaceHoudiniCSV, GetRowIndexesForParticleAtTime);
-DEFINE_NDI_RAW_FUNC_BINDER(UNiagaraDataInterfaceHoudiniCSV, GetParticlePositionAtTime);
-DEFINE_NDI_RAW_FUNC_BINDER(UNiagaraDataInterfaceHoudiniCSV, GetParticleValueAtTime);
-DEFINE_NDI_RAW_FUNC_BINDER(UNiagaraDataInterfaceHoudiniCSV, GetParticleVectorValueAtTime);
-DEFINE_NDI_RAW_FUNC_BINDER(UNiagaraDataInterfaceHoudiniCSV, GetParticleLifeAtTime);
+DEFINE_NDI_RAW_FUNC_BINDER(UNiagaraDataInterfaceHoudiniCSV, GetPointIDsToSpawnAtTime);
+DEFINE_NDI_RAW_FUNC_BINDER(UNiagaraDataInterfaceHoudiniCSV, GetRowIndexesForPointAtTime);
+DEFINE_NDI_RAW_FUNC_BINDER(UNiagaraDataInterfaceHoudiniCSV, GetPointPositionAtTime);
+DEFINE_NDI_RAW_FUNC_BINDER(UNiagaraDataInterfaceHoudiniCSV, GetPointValueAtTime);
+DEFINE_NDI_RAW_FUNC_BINDER(UNiagaraDataInterfaceHoudiniCSV, GetPointVectorValueAtTime);
+DEFINE_NDI_RAW_FUNC_BINDER(UNiagaraDataInterfaceHoudiniCSV, GetPointVectorValueAtTimeEx);
+DEFINE_NDI_RAW_FUNC_BINDER(UNiagaraDataInterfaceHoudiniCSV, GetPointLife);
+//DEFINE_NDI_RAW_FUNC_BINDER(UNiagaraDataInterfaceHoudiniCSV, GetPointLifeAtTime);
+DEFINE_NDI_RAW_FUNC_BINDER(UNiagaraDataInterfaceHoudiniCSV, GetPointType);
 void UNiagaraDataInterfaceHoudiniCSV::GetVMExternalFunction(const FVMExternalFunctionBindingInfo& BindingInfo, void* InstanceData, FVMExternalFunction &OutFunc)
 {
-    if (BindingInfo.Name == TEXT("GetCSVFloatValue") && BindingInfo.GetNumInputs() == 2 && BindingInfo.GetNumOutputs() == 1)
+    if (BindingInfo.Name == TEXT("GetFloatValue") && BindingInfo.GetNumInputs() == 2 && BindingInfo.GetNumOutputs() == 1)
     {
-		TNDIParamBinder<0, int32, TNDIParamBinder<1, int32, NDI_RAW_FUNC_BINDER(UNiagaraDataInterfaceHoudiniCSV, GetCSVFloatValue)>>::Bind(this, BindingInfo, InstanceData, OutFunc);
+		TNDIParamBinder<0, int32, TNDIParamBinder<1, int32, NDI_RAW_FUNC_BINDER(UNiagaraDataInterfaceHoudiniCSV, GetFloatValue)>>::Bind(this, BindingInfo, InstanceData, OutFunc);
     }
-    /*else if (BindingInfo.Name == TEXT("GetCSVFloatValueByString") && BindingInfo.GetNumInputs() == 2 && BindingInfo.GetNumOutputs() == 1)
+    /*else if (BindingInfo.Name == TEXT("GetFloatValueByString") && BindingInfo.GetNumInputs() == 2 && BindingInfo.GetNumOutputs() == 1)
     {
-		TNDIParamBinder<0, float, TNDIParamBinder<1, FString, NDI_RAW_FUNC_BINDER(UNiagaraDataInterfaceHoudiniCSV, GetCSVFloatValueByString)>>::Bind(this, BindingInfo, InstanceData, OutFunc);
+		TNDIParamBinder<0, float, TNDIParamBinder<1, FString, NDI_RAW_FUNC_BINDER(UNiagaraDataInterfaceHoudiniCSV, GetFloatValueByString)>>::Bind(this, BindingInfo, InstanceData, OutFunc);
     }*/
-    else if (BindingInfo.Name == TEXT("GetCSVVectorValue") && BindingInfo.GetNumInputs() == 2 && BindingInfo.GetNumOutputs() == 3)
+    else if (BindingInfo.Name == TEXT("GetVectorValue") && BindingInfo.GetNumInputs() == 2 && BindingInfo.GetNumOutputs() == 3)
     {
-		TNDIParamBinder<0, int32, TNDIParamBinder<1, int32, NDI_RAW_FUNC_BINDER(UNiagaraDataInterfaceHoudiniCSV, GetCSVVectorValue)>>::Bind(this, BindingInfo, InstanceData, OutFunc);
-    }    
-    else if (BindingInfo.Name == TEXT("GetCSVPosition") && BindingInfo.GetNumInputs() == 1 && BindingInfo.GetNumOutputs() == 3)
+		TNDIParamBinder<0, int32, TNDIParamBinder<1, int32, NDI_RAW_FUNC_BINDER(UNiagaraDataInterfaceHoudiniCSV, GetVectorValue)>>::Bind(this, BindingInfo, InstanceData, OutFunc);
+    }
+	else if (BindingInfo.Name == TEXT("GetVectorValueEx") && BindingInfo.GetNumInputs() == 4 && BindingInfo.GetNumOutputs() == 3)
+	{
+		TNDIParamBinder<0, int32, TNDIParamBinder<1, int32, TNDIParamBinder<2, bool, TNDIParamBinder<3, bool, NDI_RAW_FUNC_BINDER(UNiagaraDataInterfaceHoudiniCSV, GetVectorValueEx)>>>>::Bind(this, BindingInfo, InstanceData, OutFunc);
+	}
+    else if (BindingInfo.Name == TEXT("GetPosition") && BindingInfo.GetNumInputs() == 1 && BindingInfo.GetNumOutputs() == 3)
     {
-		TNDIParamBinder<0, int32, NDI_RAW_FUNC_BINDER(UNiagaraDataInterfaceHoudiniCSV, GetCSVPosition)>::Bind(this, BindingInfo, InstanceData, OutFunc);
+		TNDIParamBinder<0, int32, NDI_RAW_FUNC_BINDER(UNiagaraDataInterfaceHoudiniCSV, GetPosition)>::Bind(this, BindingInfo, InstanceData, OutFunc);
     }
-    else if (BindingInfo.Name == TEXT("GetCSVNormal") && BindingInfo.GetNumInputs() == 1 && BindingInfo.GetNumOutputs() == 3)
+    else if (BindingInfo.Name == TEXT("GetNormal") && BindingInfo.GetNumInputs() == 1 && BindingInfo.GetNumOutputs() == 3)
     {
-		TNDIParamBinder<0, int32, NDI_RAW_FUNC_BINDER(UNiagaraDataInterfaceHoudiniCSV, GetCSVNormal)>::Bind(this, BindingInfo, InstanceData, OutFunc);
+		TNDIParamBinder<0, int32, NDI_RAW_FUNC_BINDER(UNiagaraDataInterfaceHoudiniCSV, GetNormal)>::Bind(this, BindingInfo, InstanceData, OutFunc);
     }
-    else if (BindingInfo.Name == TEXT("GetCSVTime") && BindingInfo.GetNumInputs() == 1 && BindingInfo.GetNumOutputs() == 1)
+    else if (BindingInfo.Name == TEXT("GetTime") && BindingInfo.GetNumInputs() == 1 && BindingInfo.GetNumOutputs() == 1)
     {
-		TNDIParamBinder<0, int32, NDI_RAW_FUNC_BINDER(UNiagaraDataInterfaceHoudiniCSV, GetCSVTime)>::Bind(this, BindingInfo, InstanceData, OutFunc);
+		TNDIParamBinder<0, int32, NDI_RAW_FUNC_BINDER(UNiagaraDataInterfaceHoudiniCSV, GetTime)>::Bind(this, BindingInfo, InstanceData, OutFunc);
     }
-	else if (BindingInfo.Name == TEXT("GetCSVVelocity") && BindingInfo.GetNumInputs() == 1 && BindingInfo.GetNumOutputs() == 3)
+	else if (BindingInfo.Name == TEXT("GetVelocity") && BindingInfo.GetNumInputs() == 1 && BindingInfo.GetNumOutputs() == 3)
 	{
-		TNDIParamBinder<0, int32, NDI_RAW_FUNC_BINDER(UNiagaraDataInterfaceHoudiniCSV, GetCSVVelocity)>::Bind(this, BindingInfo, InstanceData, OutFunc);
+		TNDIParamBinder<0, int32, NDI_RAW_FUNC_BINDER(UNiagaraDataInterfaceHoudiniCSV, GetVelocity)>::Bind(this, BindingInfo, InstanceData, OutFunc);
 	}
-	else if (BindingInfo.Name == TEXT("GetCSVColor") && BindingInfo.GetNumInputs() == 1 && BindingInfo.GetNumOutputs() == 4)
+	else if (BindingInfo.Name == TEXT("GetColor") && BindingInfo.GetNumInputs() == 1 && BindingInfo.GetNumOutputs() == 4)
 	{
-		TNDIParamBinder<0, int32, NDI_RAW_FUNC_BINDER(UNiagaraDataInterfaceHoudiniCSV, GetCSVColor)>::Bind(this, BindingInfo, InstanceData, OutFunc);
+		TNDIParamBinder<0, int32, NDI_RAW_FUNC_BINDER(UNiagaraDataInterfaceHoudiniCSV, GetColor)>::Bind(this, BindingInfo, InstanceData, OutFunc);
 	}
-    else if (BindingInfo.Name == TEXT("GetCSVPositionAndTime") && BindingInfo.GetNumInputs() == 1 && BindingInfo.GetNumOutputs() == 4)
+    else if (BindingInfo.Name == TEXT("GetPositionAndTime") && BindingInfo.GetNumInputs() == 1 && BindingInfo.GetNumOutputs() == 4)
     {
-		TNDIParamBinder<0, int32, NDI_RAW_FUNC_BINDER(UNiagaraDataInterfaceHoudiniCSV, GetCSVPositionAndTime)>::Bind(this, BindingInfo, InstanceData, OutFunc);
+		TNDIParamBinder<0, int32, NDI_RAW_FUNC_BINDER(UNiagaraDataInterfaceHoudiniCSV, GetPositionAndTime)>::Bind(this, BindingInfo, InstanceData, OutFunc);
     }
-    else if ( BindingInfo.Name == TEXT("GetNumberOfParticlesInCSV") && BindingInfo.GetNumInputs() == 0 && BindingInfo.GetNumOutputs() == 1 )
+    else if ( BindingInfo.Name == TEXT("GetNumberOfPoints") && BindingInfo.GetNumInputs() == 0 && BindingInfo.GetNumOutputs() == 1 )
     {
-		OutFunc = FVMExternalFunction::CreateUObject(this, &UNiagaraDataInterfaceHoudiniCSV::GetNumberOfParticlesInCSV);
+		OutFunc = FVMExternalFunction::CreateUObject(this, &UNiagaraDataInterfaceHoudiniCSV::GetNumberOfPoints);
     }
-	else if (BindingInfo.Name == TEXT("GetNumberOfRowsInCSV") && BindingInfo.GetNumInputs() == 0 && BindingInfo.GetNumOutputs() == 1)
+	else if (BindingInfo.Name == TEXT("GetNumberOfRows") && BindingInfo.GetNumInputs() == 0 && BindingInfo.GetNumOutputs() == 1)
 	{
-		OutFunc = FVMExternalFunction::CreateUObject(this, &UNiagaraDataInterfaceHoudiniCSV::GetNumberOfRowsInCSV);
+		OutFunc = FVMExternalFunction::CreateUObject(this, &UNiagaraDataInterfaceHoudiniCSV::GetNumberOfRows);
 	}
-	else if (BindingInfo.Name == TEXT("GetNumberOfColumnsInCSV") && BindingInfo.GetNumInputs() == 0 && BindingInfo.GetNumOutputs() == 1)
+	else if (BindingInfo.Name == TEXT("GetNumberOfColumns") && BindingInfo.GetNumInputs() == 0 && BindingInfo.GetNumOutputs() == 1)
 	{
-		OutFunc = FVMExternalFunction::CreateUObject(this, &UNiagaraDataInterfaceHoudiniCSV::GetNumberOfColumnsInCSV);
+		OutFunc = FVMExternalFunction::CreateUObject(this, &UNiagaraDataInterfaceHoudiniCSV::GetNumberOfColumns);
 	}
     else if (BindingInfo.Name == TEXT("GetLastRowIndexAtTime") && BindingInfo.GetNumInputs() == 1 && BindingInfo.GetNumOutputs() == 1)
     {
 		TNDIParamBinder<0, float, NDI_RAW_FUNC_BINDER(UNiagaraDataInterfaceHoudiniCSV, GetLastRowIndexAtTime)>::Bind(this, BindingInfo, InstanceData, OutFunc);
     }
-    else if (BindingInfo.Name == TEXT("GetParticleIndexesToSpawnAtTime") && BindingInfo.GetNumInputs() == 1 && BindingInfo.GetNumOutputs() == 3)
+    else if (BindingInfo.Name == TEXT("GetPointIDsToSpawnAtTime") && BindingInfo.GetNumInputs() == 1 && BindingInfo.GetNumOutputs() == 3)
     {
-		TNDIParamBinder<0, float, NDI_RAW_FUNC_BINDER(UNiagaraDataInterfaceHoudiniCSV, GetParticleIndexesToSpawnAtTime)>::Bind(this, BindingInfo, InstanceData, OutFunc);
+		TNDIParamBinder<0, float, NDI_RAW_FUNC_BINDER(UNiagaraDataInterfaceHoudiniCSV, GetPointIDsToSpawnAtTime)>::Bind(this, BindingInfo, InstanceData, OutFunc);
     }
-	else if (BindingInfo.Name == TEXT("GetRowIndexesForParticleAtTime") && BindingInfo.GetNumInputs() == 2 && BindingInfo.GetNumOutputs() == 3)
+	else if (BindingInfo.Name == TEXT("GetRowIndexesForPointAtTime") && BindingInfo.GetNumInputs() == 2 && BindingInfo.GetNumOutputs() == 3)
 	{
-		TNDIParamBinder<0, int32, TNDIParamBinder<1, float, NDI_RAW_FUNC_BINDER(UNiagaraDataInterfaceHoudiniCSV, GetRowIndexesForParticleAtTime)>>::Bind(this, BindingInfo, InstanceData, OutFunc);
+		TNDIParamBinder<0, int32, TNDIParamBinder<1, float, NDI_RAW_FUNC_BINDER(UNiagaraDataInterfaceHoudiniCSV, GetRowIndexesForPointAtTime)>>::Bind(this, BindingInfo, InstanceData, OutFunc);
 	}
-	else if (BindingInfo.Name == TEXT("GetParticlePositionAtTime") && BindingInfo.GetNumInputs() == 2 && BindingInfo.GetNumOutputs() == 3)
+	else if (BindingInfo.Name == TEXT("GetPointPositionAtTime") && BindingInfo.GetNumInputs() == 2 && BindingInfo.GetNumOutputs() == 3)
 	{
-		TNDIParamBinder<0, int32, TNDIParamBinder<1, float, NDI_RAW_FUNC_BINDER(UNiagaraDataInterfaceHoudiniCSV, GetParticlePositionAtTime)>>::Bind(this, BindingInfo, InstanceData, OutFunc);
+		TNDIParamBinder<0, int32, TNDIParamBinder<1, float, NDI_RAW_FUNC_BINDER(UNiagaraDataInterfaceHoudiniCSV, GetPointPositionAtTime)>>::Bind(this, BindingInfo, InstanceData, OutFunc);
 	}
-	else if (BindingInfo.Name == TEXT("GetParticleValueAtTime") && BindingInfo.GetNumInputs() == 3 && BindingInfo.GetNumOutputs() == 1)
+	else if (BindingInfo.Name == TEXT("GetPointValueAtTime") && BindingInfo.GetNumInputs() == 3 && BindingInfo.GetNumOutputs() == 1)
 	{
-		TNDIParamBinder<0, int32, TNDIParamBinder<1, int32, TNDIParamBinder<2, float, NDI_RAW_FUNC_BINDER(UNiagaraDataInterfaceHoudiniCSV, GetParticleValueAtTime)>>>::Bind(this, BindingInfo, InstanceData, OutFunc);
+		TNDIParamBinder<0, int32, TNDIParamBinder<1, int32, TNDIParamBinder<2, float, NDI_RAW_FUNC_BINDER(UNiagaraDataInterfaceHoudiniCSV, GetPointValueAtTime)>>>::Bind(this, BindingInfo, InstanceData, OutFunc);
 	}
-	else if (BindingInfo.Name == TEXT("GetParticleVectorValueAtTime") && BindingInfo.GetNumInputs() == 3 && BindingInfo.GetNumOutputs() == 3)
+	else if (BindingInfo.Name == TEXT("GetPointVectorValueAtTime") && BindingInfo.GetNumInputs() == 3 && BindingInfo.GetNumOutputs() == 3)
 	{
-		TNDIParamBinder<0, int32, TNDIParamBinder<1, int32, TNDIParamBinder<2, float, NDI_RAW_FUNC_BINDER(UNiagaraDataInterfaceHoudiniCSV, GetParticleVectorValueAtTime)>>>::Bind(this, BindingInfo, InstanceData, OutFunc);
+		TNDIParamBinder<0, int32, TNDIParamBinder<1, int32, TNDIParamBinder<2, float, NDI_RAW_FUNC_BINDER(UNiagaraDataInterfaceHoudiniCSV, GetPointVectorValueAtTime)>>>::Bind(this, BindingInfo, InstanceData, OutFunc);
 	}
-	else if (BindingInfo.Name == TEXT("GetParticleLifeAtTime") && BindingInfo.GetNumInputs() == 2 && BindingInfo.GetNumOutputs() == 1)
+	else if (BindingInfo.Name == TEXT("GetPointVectorValueAtTimeEx") && BindingInfo.GetNumInputs() == 5 && BindingInfo.GetNumOutputs() == 3)
 	{
-		TNDIParamBinder<0, int32, TNDIParamBinder<1, float, NDI_RAW_FUNC_BINDER(UNiagaraDataInterfaceHoudiniCSV, GetParticleLifeAtTime)>>::Bind(this, BindingInfo, InstanceData, OutFunc);
+		TNDIParamBinder<0, int32, TNDIParamBinder<1, int32, TNDIParamBinder<2, float, TNDIParamBinder<3, bool, TNDIParamBinder<4, bool, NDI_RAW_FUNC_BINDER(UNiagaraDataInterfaceHoudiniCSV, GetPointVectorValueAtTimeEx)>>>>>::Bind(this, BindingInfo, InstanceData, OutFunc);
+	}
+	else if (BindingInfo.Name == TEXT("GetPointLife") && BindingInfo.GetNumInputs() == 1 && BindingInfo.GetNumOutputs() == 1)
+	{
+		TNDIParamBinder<0, int32, NDI_RAW_FUNC_BINDER(UNiagaraDataInterfaceHoudiniCSV, GetPointLife)>::Bind(this, BindingInfo, InstanceData, OutFunc);
+	}
+	/*
+	else if (BindingInfo.Name == TEXT("GetPointLifeAtTime") && BindingInfo.GetNumInputs() == 2 && BindingInfo.GetNumOutputs() == 1)
+	{
+		TNDIParamBinder<0, int32, TNDIParamBinder<1, float, NDI_RAW_FUNC_BINDER(UNiagaraDataInterfaceHoudiniCSV, GetPointLifeAtTime)>>::Bind(this, BindingInfo, InstanceData, OutFunc);
+	}
+	*/
+	else if (BindingInfo.Name == TEXT("GetPointType") && BindingInfo.GetNumInputs() == 1 && BindingInfo.GetNumOutputs() == 1)
+	{
+		TNDIParamBinder<0, int32, NDI_RAW_FUNC_BINDER(UNiagaraDataInterfaceHoudiniCSV, GetPointType)>::Bind(this, BindingInfo, InstanceData, OutFunc);
 	}
     else
     {
@@ -613,7 +708,7 @@ void UNiagaraDataInterfaceHoudiniCSV::GetVMExternalFunction(const FVMExternalFun
 }
 
 template<typename RowParamType, typename ColParamType>
-void UNiagaraDataInterfaceHoudiniCSV::GetCSVFloatValue(FVectorVMContext& Context)
+void UNiagaraDataInterfaceHoudiniCSV::GetFloatValue(FVectorVMContext& Context)
 {
     RowParamType RowParam(Context);
     ColParamType ColParam(Context);
@@ -626,8 +721,8 @@ void UNiagaraDataInterfaceHoudiniCSV::GetCSVFloatValue(FVectorVMContext& Context
 		int32 col = ColParam.Get();
 	
 		float value = 0.0f;
-		if ( CSVFile )
-			CSVFile->GetCSVFloatValue( row, col, value );
+		if ( HoudiniCSVAsset )
+			HoudiniCSVAsset->GetFloatValue( row, col, value );
 
 		*OutValue.GetDest() = value;
 		RowParam.Advance();
@@ -637,7 +732,7 @@ void UNiagaraDataInterfaceHoudiniCSV::GetCSVFloatValue(FVectorVMContext& Context
 }
 
 template<typename RowParamType, typename ColParamType>
-void UNiagaraDataInterfaceHoudiniCSV::GetCSVVectorValue( FVectorVMContext& Context )
+void UNiagaraDataInterfaceHoudiniCSV::GetVectorValue( FVectorVMContext& Context )
 {
     RowParamType RowParam(Context);
     ColParamType ColParam(Context);
@@ -652,8 +747,8 @@ void UNiagaraDataInterfaceHoudiniCSV::GetCSVVectorValue( FVectorVMContext& Conte
 		int32 col = ColParam.Get();
 
 		FVector V = FVector::ZeroVector;
-		if ( CSVFile )
-			CSVFile->GetCSVVectorValue(row, col, V);
+		if ( HoudiniCSVAsset )
+			HoudiniCSVAsset->GetVectorValue(row, col, V);
 
 		*OutVectorX.GetDest() = V.X;
 		*OutVectorY.GetDest() = V.Y;
@@ -666,9 +761,47 @@ void UNiagaraDataInterfaceHoudiniCSV::GetCSVVectorValue( FVectorVMContext& Conte
 		OutVectorZ.Advance();
     }
 }
+
+template<typename RowParamType, typename ColParamType, typename DoSwapParamType, typename DoScaleParamType >
+void UNiagaraDataInterfaceHoudiniCSV::GetVectorValueEx(FVectorVMContext& Context)
+{
+	RowParamType RowParam(Context);
+	ColParamType ColParam(Context);
+	DoSwapParamType DoSwapParam(Context);
+	DoScaleParamType DoScaleParam(Context);
+
+	FRegisterHandler<float> OutVectorX(Context);
+	FRegisterHandler<float> OutVectorY(Context);
+	FRegisterHandler<float> OutVectorZ(Context);
+
+	for (int32 i = 0; i < Context.NumInstances; ++i)
+	{
+		int32 row = RowParam.Get();
+		int32 col = ColParam.Get();
+
+		bool DoSwap = DoSwapParam.Get();
+		bool DoScale = DoScaleParam.Get();
+
+		FVector V = FVector::ZeroVector;
+		if (HoudiniCSVAsset)
+			HoudiniCSVAsset->GetVectorValue(row, col, V, DoSwap, DoScale);
+
+		*OutVectorX.GetDest() = V.X;
+		*OutVectorY.GetDest() = V.Y;
+		*OutVectorZ.GetDest() = V.Z;
+
+		RowParam.Advance();
+		ColParam.Advance();
+		DoSwapParam.Advance();
+		DoScaleParam.Advance();
+		OutVectorX.Advance();
+		OutVectorY.Advance();
+		OutVectorZ.Advance();
+	}
+}
 /*
 template<typename RowParamType, typename ColTitleParamType>
-void UNiagaraDataInterfaceHoudiniCSV::GetCSVFloatValueByString(FVectorVMContext& Context)
+void UNiagaraDataInterfaceHoudiniCSV::GetFloatValueByString(FVectorVMContext& Context)
 {
     RowParamType RowParam(Context);
     ColTitleParamType ColTitleParam(Context);
@@ -693,7 +826,7 @@ void UNiagaraDataInterfaceHoudiniCSV::GetCSVFloatValueByString(FVectorVMContext&
 */
 
 template<typename RowParamType>
-void UNiagaraDataInterfaceHoudiniCSV::GetCSVPosition(FVectorVMContext& Context)
+void UNiagaraDataInterfaceHoudiniCSV::GetPosition(FVectorVMContext& Context)
 {
 	RowParamType RowParam(Context);
     FRegisterHandler<float> OutSampleX(Context);
@@ -705,8 +838,8 @@ void UNiagaraDataInterfaceHoudiniCSV::GetCSVPosition(FVectorVMContext& Context)
 		int32 row = RowParam.Get();
 
 		FVector V = FVector::ZeroVector;
-		if ( CSVFile )
-			CSVFile->GetCSVPositionValue( row, V );
+		if ( HoudiniCSVAsset )
+			HoudiniCSVAsset->GetPositionValue( row, V );
 
 		*OutSampleX.GetDest() = V.X;
 		*OutSampleY.GetDest() = V.Y;
@@ -719,7 +852,7 @@ void UNiagaraDataInterfaceHoudiniCSV::GetCSVPosition(FVectorVMContext& Context)
 }
 
 template<typename RowParamType>
-void UNiagaraDataInterfaceHoudiniCSV::GetCSVNormal(FVectorVMContext& Context)
+void UNiagaraDataInterfaceHoudiniCSV::GetNormal(FVectorVMContext& Context)
 {
 	RowParamType RowParam(Context);
     FRegisterHandler<float> OutSampleX(Context);
@@ -731,8 +864,8 @@ void UNiagaraDataInterfaceHoudiniCSV::GetCSVNormal(FVectorVMContext& Context)
 		int32 row = RowParam.Get();
 
 		FVector V = FVector::ZeroVector;
-		if ( CSVFile )
-			CSVFile->GetCSVNormalValue( row, V );
+		if ( HoudiniCSVAsset )
+			HoudiniCSVAsset->GetNormalValue( row, V );
 
 		*OutSampleX.GetDest() = V.X;
 		*OutSampleY.GetDest() = V.Y;
@@ -745,7 +878,7 @@ void UNiagaraDataInterfaceHoudiniCSV::GetCSVNormal(FVectorVMContext& Context)
 }
 
 template<typename RowParamType>
-void UNiagaraDataInterfaceHoudiniCSV::GetCSVTime(FVectorVMContext& Context)
+void UNiagaraDataInterfaceHoudiniCSV::GetTime(FVectorVMContext& Context)
 {
 	RowParamType RowParam(Context);
     FRegisterHandler<float> OutValue(Context);
@@ -755,8 +888,8 @@ void UNiagaraDataInterfaceHoudiniCSV::GetCSVTime(FVectorVMContext& Context)
 		int32 row = RowParam.Get();
 
 		float value = 0.0f;
-		if ( CSVFile )
-			CSVFile->GetCSVTimeValue( row, value );
+		if ( HoudiniCSVAsset )
+			HoudiniCSVAsset->GetTimeValue( row, value );
 
 		*OutValue.GetDest() = value;
 		RowParam.Advance();
@@ -765,7 +898,7 @@ void UNiagaraDataInterfaceHoudiniCSV::GetCSVTime(FVectorVMContext& Context)
 }
 
 template<typename RowParamType>
-void UNiagaraDataInterfaceHoudiniCSV::GetCSVVelocity(FVectorVMContext& Context)
+void UNiagaraDataInterfaceHoudiniCSV::GetVelocity(FVectorVMContext& Context)
 {
 	RowParamType RowParam(Context);
 	FRegisterHandler<float> OutSampleX(Context);
@@ -777,8 +910,8 @@ void UNiagaraDataInterfaceHoudiniCSV::GetCSVVelocity(FVectorVMContext& Context)
 		int32 row = RowParam.Get();
 
 		FVector V = FVector::ZeroVector;
-		if ( CSVFile )
-			CSVFile->GetCSVVelocityValue( row, V );
+		if ( HoudiniCSVAsset )
+			HoudiniCSVAsset->GetVelocityValue( row, V );
 
 		*OutSampleX.GetDest() = V.X;
 		*OutSampleY.GetDest() = V.Y;
@@ -791,7 +924,7 @@ void UNiagaraDataInterfaceHoudiniCSV::GetCSVVelocity(FVectorVMContext& Context)
 }
 
 template<typename RowParamType>
-void UNiagaraDataInterfaceHoudiniCSV::GetCSVColor(FVectorVMContext& Context)
+void UNiagaraDataInterfaceHoudiniCSV::GetColor(FVectorVMContext& Context)
 {
 	RowParamType RowParam(Context);
 	FRegisterHandler<float> OutSampleR(Context);
@@ -804,8 +937,8 @@ void UNiagaraDataInterfaceHoudiniCSV::GetCSVColor(FVectorVMContext& Context)
 		int32 row = RowParam.Get();
 
 		FLinearColor C = FLinearColor::White;
-		if ( CSVFile )
-			CSVFile->GetCSVColorValue( row, C );
+		if ( HoudiniCSVAsset )
+			HoudiniCSVAsset->GetColorValue( row, C );
 
 		*OutSampleR.GetDest() = C.R;
 		*OutSampleG.GetDest() = C.G;
@@ -819,7 +952,7 @@ void UNiagaraDataInterfaceHoudiniCSV::GetCSVColor(FVectorVMContext& Context)
 	}
 }
 
-// Returns the last index of the particles that should be spawned at time t
+// Returns the last index of the points that should be spawned at time t
 template<typename TimeParamType>
 void UNiagaraDataInterfaceHoudiniCSV::GetLastRowIndexAtTime(FVectorVMContext& Context)
 {
@@ -831,8 +964,8 @@ void UNiagaraDataInterfaceHoudiniCSV::GetLastRowIndexAtTime(FVectorVMContext& Co
 		float t = TimeParam.Get();
 
 		int32 value = 0;
-		if ( CSVFile )
-			CSVFile->GetLastRowIndexAtTime( t, value );
+		if ( HoudiniCSVAsset )
+			HoudiniCSVAsset->GetLastRowIndexAtTime( t, value );
 
 		*OutValue.GetDest() = value;
 		TimeParam.Advance();
@@ -840,9 +973,9 @@ void UNiagaraDataInterfaceHoudiniCSV::GetLastRowIndexAtTime(FVectorVMContext& Co
     }
 }
 
-// Returns the last index of the particles that should be spawned at time t
+// Returns the last index of the points that should be spawned at time t
 template<typename TimeParamType>
-void UNiagaraDataInterfaceHoudiniCSV::GetParticleIndexesToSpawnAtTime( FVectorVMContext& Context )
+void UNiagaraDataInterfaceHoudiniCSV::GetPointIDsToSpawnAtTime( FVectorVMContext& Context )
 {
     TimeParamType TimeParam( Context );
     FRegisterHandler<int32> OutMinValue( Context );
@@ -855,9 +988,9 @@ void UNiagaraDataInterfaceHoudiniCSV::GetParticleIndexesToSpawnAtTime( FVectorVM
 
 		int32 value = 0;
 		int32 min = 0, max = 0, count = 0;
-		if ( CSVFile )
+		if ( HoudiniCSVAsset )
 		{
-			if ( !CSVFile->GetLastParticleIndexAtTime( t, value ) )
+			if ( !HoudiniCSVAsset->GetLastPointIDToSpawnAtTime( t, value ) )
 			{
 				// The CSV file doesn't have time informations, so always return all points in the file
 				min = 0;
@@ -867,36 +1000,36 @@ void UNiagaraDataInterfaceHoudiniCSV::GetParticleIndexesToSpawnAtTime( FVectorVM
 			else
 			{
 				// The CSV file has time informations
-				// First, detect if we need to reset LastSpawnParticleID (after a loop of the emitter)
-				if ( value < LastSpawnedParticleID )
-					LastSpawnedParticleID = -1;
+				// First, detect if we need to reset LastSpawnedPointID (after a loop of the emitter)
+				if ( value < LastSpawnedPointID )
+					LastSpawnedPointID = -1;
 
 				if ( value < 0 )
 				{
-					// Nothing to spawn, t is lower than the particle's time
-					LastSpawnedParticleID = -1;
+					// Nothing to spawn, t is lower than the point's time
+					LastSpawnedPointID = -1;
 				}
 				else
 				{
 					// The last time value in the CSV is lower than t, spawn everything if we didnt already!
-					if ( value >= CSVFile->GetNumberOfParticlesInCSV() )
+					if ( value >= HoudiniCSVAsset->GetNumberOfPoints() )
 						value = value - 1;
 
-					if ( value == LastSpawnedParticleID)
+					if ( value == LastSpawnedPointID)
 					{
-						// We dont have any new particle to spawn
+						// We dont have any new point to spawn
 						min = value;
 						max = value;
 						count = 0;
 					}
 					else
 					{
-						// We have particles to spawn at time t
-						min = LastSpawnedParticleID + 1;
+						// We have points to spawn at time t
+						min = LastSpawnedPointID + 1;
 						max = value;
 						count = max - min + 1;
 
-						LastSpawnedParticleID = max;
+						LastSpawnedPointID = max;
 					}
 				}
 			}
@@ -914,7 +1047,7 @@ void UNiagaraDataInterfaceHoudiniCSV::GetParticleIndexesToSpawnAtTime( FVectorVM
 }
 
 template<typename RowParamType>
-void UNiagaraDataInterfaceHoudiniCSV::GetCSVPositionAndTime(FVectorVMContext& Context)
+void UNiagaraDataInterfaceHoudiniCSV::GetPositionAndTime(FVectorVMContext& Context)
 {
     RowParamType RowParam(Context);
 
@@ -929,10 +1062,10 @@ void UNiagaraDataInterfaceHoudiniCSV::GetCSVPositionAndTime(FVectorVMContext& Co
 
 		float timeValue = 0.0f;
 		FVector posVector = FVector::ZeroVector;
-		if ( CSVFile )
+		if ( HoudiniCSVAsset )
 		{
-			CSVFile->GetCSVTimeValue( row, timeValue);
-			CSVFile->GetCSVPositionValue( row, posVector);
+			HoudiniCSVAsset->GetTimeValue( row, timeValue);
+			HoudiniCSVAsset->GetPositionValue( row, posVector);
 		}
 
 		*OutPosX.GetDest() = posVector.X;
@@ -949,10 +1082,10 @@ void UNiagaraDataInterfaceHoudiniCSV::GetCSVPositionAndTime(FVectorVMContext& Co
     }
 }
 
-template<typename NParamType, typename TimeParamType>
-void UNiagaraDataInterfaceHoudiniCSV::GetRowIndexesForParticleAtTime(FVectorVMContext& Context)
+template<typename PointIDParamType, typename TimeParamType>
+void UNiagaraDataInterfaceHoudiniCSV::GetRowIndexesForPointAtTime(FVectorVMContext& Context)
 {
-	NParamType NParam(Context);
+	PointIDParamType PointIDParam(Context);
 	TimeParamType TimeParam(Context);
 
 	FRegisterHandler<int32> OutPrevIndex(Context);
@@ -961,22 +1094,22 @@ void UNiagaraDataInterfaceHoudiniCSV::GetRowIndexesForParticleAtTime(FVectorVMCo
 
 	for (int32 i = 0; i < Context.NumInstances; ++i)
     {
-		int32 N = NParam.Get();
+		int32 PointID = PointIDParam.Get();
 		float time = TimeParam.Get();
 
 		float weight = 0.0f;
 		int32 prevIdx = 0;
 		int32 nextIdx = 0;
-		if ( CSVFile )
+		if ( HoudiniCSVAsset )
 		{
-			CSVFile->GetParticleLineIndexAtTime( N, time, prevIdx, nextIdx, weight );
+			HoudiniCSVAsset->GetRowIndexesForPointAtTime( PointID, time, prevIdx, nextIdx, weight );
 		}
 
 		*OutPrevIndex.GetDest() = prevIdx;
 		*OutNextIndex.GetDest() = nextIdx;
 		*OutWeightValue.GetDest() = weight;
 
-		NParam.Advance();
+		PointIDParam.Advance();
 		TimeParam.Advance();
 		OutPrevIndex.Advance();
 		OutNextIndex.Advance();
@@ -984,10 +1117,10 @@ void UNiagaraDataInterfaceHoudiniCSV::GetRowIndexesForParticleAtTime(FVectorVMCo
     }
 }
 
-template<typename NParamType, typename TimeParamType>
-void UNiagaraDataInterfaceHoudiniCSV::GetParticlePositionAtTime( FVectorVMContext& Context )
+template<typename PointIDParamType, typename TimeParamType>
+void UNiagaraDataInterfaceHoudiniCSV::GetPointPositionAtTime( FVectorVMContext& Context )
 {
-	NParamType NParam(Context);
+	PointIDParamType PointIDParam(Context);
 	TimeParamType TimeParam(Context);
 
 	FRegisterHandler<float> OutPosX(Context);
@@ -996,20 +1129,20 @@ void UNiagaraDataInterfaceHoudiniCSV::GetParticlePositionAtTime( FVectorVMContex
 
 	for (int32 i = 0; i < Context.NumInstances; ++i)
     {
-		int32 N = NParam.Get();
+		int32 PointID = PointIDParam.Get();
 		float time = TimeParam.Get();
 
 		FVector posVector = FVector::ZeroVector;
-		if ( CSVFile )
+		if ( HoudiniCSVAsset )
 		{
-			CSVFile->GetParticlePositionAtTime(N, time, posVector);
+			HoudiniCSVAsset->GetPointPositionAtTime(PointID, time, posVector);
 		}		
 
 		*OutPosX.GetDest() = posVector.X;
 		*OutPosY.GetDest() = posVector.Y;
 		*OutPosZ.GetDest() = posVector.Z;
 
-		NParam.Advance();
+		PointIDParam.Advance();
 		TimeParam.Advance();
 		OutPosX.Advance();
 		OutPosY.Advance();
@@ -1017,10 +1150,10 @@ void UNiagaraDataInterfaceHoudiniCSV::GetParticlePositionAtTime( FVectorVMContex
     }
 }
 
-template<typename NParamType, typename ColParamType, typename TimeParamType>
-void UNiagaraDataInterfaceHoudiniCSV::GetParticleValueAtTime(FVectorVMContext& Context)
+template<typename PointIDParamType, typename ColParamType, typename TimeParamType>
+void UNiagaraDataInterfaceHoudiniCSV::GetPointValueAtTime(FVectorVMContext& Context)
 {
-	NParamType NParam(Context);
+	PointIDParamType PointIDParam(Context);
 	TimeParamType TimeParam(Context);
 	ColParamType ColParam(Context);
 
@@ -1028,19 +1161,19 @@ void UNiagaraDataInterfaceHoudiniCSV::GetParticleValueAtTime(FVectorVMContext& C
 
 	for (int32 i = 0; i < Context.NumInstances; ++i)
 	{
-		int32 N = NParam.Get();
+		int32 PointID = PointIDParam.Get();
 		int32 Col = ColParam.Get();
 		float time = TimeParam.Get();		
 
 		float Value = 0.0f;
-		if ( CSVFile )
+		if ( HoudiniCSVAsset )
 		{
-			CSVFile->GetParticleValueAtTime( N, Col, time, Value );
+			HoudiniCSVAsset->GetPointValueAtTime( PointID, Col, time, Value );
 		}
 
 		*OutValue.GetDest() = Value;
 
-		NParam.Advance();
+		PointIDParam.Advance();
 		ColParam.Advance();
 		TimeParam.Advance();
 
@@ -1048,10 +1181,10 @@ void UNiagaraDataInterfaceHoudiniCSV::GetParticleValueAtTime(FVectorVMContext& C
 	}
 }
 
-template<typename NParamType, typename ColParamType, typename TimeParamType>
-void UNiagaraDataInterfaceHoudiniCSV::GetParticleVectorValueAtTime(FVectorVMContext& Context)
+template<typename PointIDParamType, typename ColParamType, typename TimeParamType>
+void UNiagaraDataInterfaceHoudiniCSV::GetPointVectorValueAtTime(FVectorVMContext& Context)
 {
-	NParamType NParam(Context);
+	PointIDParamType PointIDParam(Context);
 	ColParamType ColParam(Context);
 	TimeParamType TimeParam(Context);	
 
@@ -1061,21 +1194,21 @@ void UNiagaraDataInterfaceHoudiniCSV::GetParticleVectorValueAtTime(FVectorVMCont
 
 	for (int32 i = 0; i < Context.NumInstances; ++i)
 	{
-		int32 N = NParam.Get();
+		int32 PointID = PointIDParam.Get();
 		int32 Col = ColParam.Get();
 		float time = TimeParam.Get();		
 
 		FVector posVector = FVector::ZeroVector;
-		if ( CSVFile )
+		if ( HoudiniCSVAsset )
 		{
-			CSVFile->GetParticleVectorValueAtTime(N, Col, time, posVector, true, true);
+			HoudiniCSVAsset->GetPointVectorValueAtTime( PointID, Col, time, posVector, true, true);
 		}
 
 		*OutPosX.GetDest() = posVector.X;
 		*OutPosY.GetDest() = posVector.Y;
 		*OutPosZ.GetDest() = posVector.Z;
 
-		NParam.Advance();
+		PointIDParam.Advance();
 		ColParam.Advance();
 		TimeParam.Advance();
 		
@@ -1085,53 +1218,149 @@ void UNiagaraDataInterfaceHoudiniCSV::GetParticleVectorValueAtTime(FVectorVMCont
 	}
 }
 
-template<typename NParamType, typename TimeParamType>
-void UNiagaraDataInterfaceHoudiniCSV::GetParticleLifeAtTime(FVectorVMContext& Context)
+template<typename PointIDParamType, typename ColParamType, typename TimeParamType, typename DoSwapParamType, typename DoScaleParamType >
+void UNiagaraDataInterfaceHoudiniCSV::GetPointVectorValueAtTimeEx(FVectorVMContext& Context)
+{
+	PointIDParamType PointIDParam(Context);
+	ColParamType ColParam(Context);
+	TimeParamType TimeParam(Context);
+	DoSwapParamType DoSwapParam(Context);
+	DoScaleParamType DoScaleParam(Context);
+
+	FRegisterHandler<float> OutPosX(Context);
+	FRegisterHandler<float> OutPosY(Context);
+	FRegisterHandler<float> OutPosZ(Context);
+
+	for (int32 i = 0; i < Context.NumInstances; ++i)
+	{
+		int32 PointID = PointIDParam.Get();
+		int32 Col = ColParam.Get();
+		float time = TimeParam.Get();
+
+		bool DoSwap = DoSwapParam.Get();
+		bool DoScale = DoScaleParam.Get();
+
+		FVector posVector = FVector::ZeroVector;
+		if (HoudiniCSVAsset)
+		{
+			HoudiniCSVAsset->GetPointVectorValueAtTime(PointID, Col, time, posVector, DoSwap, DoScale);
+		}
+
+		*OutPosX.GetDest() = posVector.X;
+		*OutPosY.GetDest() = posVector.Y;
+		*OutPosZ.GetDest() = posVector.Z;
+
+		PointIDParam.Advance();
+		ColParam.Advance();
+		TimeParam.Advance();
+		DoSwapParam.Advance();
+		DoScaleParam.Advance();
+
+		OutPosX.Advance();
+		OutPosY.Advance();
+		OutPosZ.Advance();
+	}
+}
+
+template<typename PointIDParamType>
+void UNiagaraDataInterfaceHoudiniCSV::GetPointLife(FVectorVMContext& Context)
 {
-	NParamType NParam(Context);
+	PointIDParamType PointIDParam(Context);
+
+	FRegisterHandler<float> OutValue(Context);
+
+	for (int32 i = 0; i < Context.NumInstances; ++i)
+	{
+		int32 PointID = PointIDParam.Get();
+
+		float Value = 0.0f;
+		if ( HoudiniCSVAsset )
+		{
+			HoudiniCSVAsset->GetPointLife(PointID, Value);
+		}
+
+		*OutValue.GetDest() = Value;
+
+		PointIDParam.Advance();
+
+		OutValue.Advance();
+	}
+}
+
+/*
+template<typename PointIDParamType, typename TimeParamType>
+void UNiagaraDataInterfaceHoudiniCSV::GetPointLifeAtTime(FVectorVMContext& Context)
+{
+	PointIDParamType PointIDParam(Context);
 	TimeParamType TimeParam(Context);
 
 	FRegisterHandler<float> OutValue(Context);
 
 	for (int32 i = 0; i < Context.NumInstances; ++i)
 	{
-		int32 N = NParam.Get();
+		int32 PointID = PointIDParam.Get();
 		float time = TimeParam.Get();
 
 		float Value = 0.0f;
-		if ( CSVFile )
+		if ( HoudiniCSVAsset )
 		{
-			CSVFile->GetParticleLifeAtTime(N, time, Value);
+			HoudiniCSVAsset->GetPointLifeAtTime(PointID, time, Value);
 		}
 
 		*OutValue.GetDest() = Value;
 
-		NParam.Advance();
+		PointIDParam.Advance();
 		TimeParam.Advance();
 
 		OutValue.Advance();
 	}
 }
+*/
+
+template<typename PointIDParamType>
+void UNiagaraDataInterfaceHoudiniCSV::GetPointType(FVectorVMContext& Context)
+{
+	PointIDParamType PointIDParam(Context);
+
+	FRegisterHandler<int32> OutValue(Context);
+
+	for (int32 i = 0; i < Context.NumInstances; ++i)
+	{
+		int32 PointID = PointIDParam.Get();
+
+		int32 Value = 0;
+		if (HoudiniCSVAsset)
+		{
+			HoudiniCSVAsset->GetPointType(PointID, Value);
+		}
+
+		*OutValue.GetDest() = Value;
+
+		PointIDParam.Advance();
+
+		OutValue.Advance();
+	}
+}
 
-void UNiagaraDataInterfaceHoudiniCSV::GetNumberOfRowsInCSV(FVectorVMContext& Context)
+void UNiagaraDataInterfaceHoudiniCSV::GetNumberOfRows(FVectorVMContext& Context)
 {
     FRegisterHandler<int32> OutNumRows(Context);
-    *OutNumRows.GetDest() = CSVFile ? CSVFile->GetNumberOfLinesInCSV() : 0;
+    *OutNumRows.GetDest() = HoudiniCSVAsset ? HoudiniCSVAsset->GetNumberOfRows() : 0;
 	OutNumRows.Advance();
 }
 
-void UNiagaraDataInterfaceHoudiniCSV::GetNumberOfColumnsInCSV(FVectorVMContext& Context)
+void UNiagaraDataInterfaceHoudiniCSV::GetNumberOfColumns(FVectorVMContext& Context)
 {
 	FRegisterHandler<int32> OutNumCols(Context);
-	*OutNumCols.GetDest() = CSVFile ? CSVFile->GetNumberOfColumnsInCSV() : 0;
+	*OutNumCols.GetDest() = HoudiniCSVAsset ? HoudiniCSVAsset->GetNumberOfColumns() : 0;
 	OutNumCols.Advance();
 }
 
-void UNiagaraDataInterfaceHoudiniCSV::GetNumberOfParticlesInCSV(FVectorVMContext& Context)
+void UNiagaraDataInterfaceHoudiniCSV::GetNumberOfPoints(FVectorVMContext& Context)
 {
-	FRegisterHandler<int32> OutNumParticles(Context);
-	*OutNumParticles.GetDest() = CSVFile ? CSVFile->GetNumberOfParticlesInCSV() : 0;
-	OutNumParticles.Advance();
+	FRegisterHandler<int32> OutNumPoints(Context);
+	*OutNumPoints.GetDest() = HoudiniCSVAsset ? HoudiniCSVAsset->GetNumberOfPoints() : 0;
+	OutNumPoints.Advance();
 }
 
 #undef LOCTEXT_NAMESPACE
diff --git a/Engine/Plugins/FX/HoudiniNiagara/Source/HoudiniNiagara/Public/HoudiniCSV.h b/Engine/Plugins/FX/HoudiniNiagara/Source/HoudiniNiagara/Public/HoudiniCSV.h
--- a/Engine/Plugins/FX/HoudiniNiagara/Source/HoudiniNiagara/Public/HoudiniCSV.h
+++ b/Engine/Plugins/FX/HoudiniNiagara/Source/HoudiniNiagara/Public/HoudiniCSV.h
@@ -31,11 +31,11 @@ DECLARE_LOG_CATEGORY_EXTERN( LogHoudiniNiagara, All, All );
 
 
 USTRUCT()
-struct FParticleIndexes
+struct FPointIndexes
 {
 	GENERATED_BODY()
 
-	// Simple structure for storing all the row indexes for a particle
+	// Simple structure for storing all the row indexes used for a given point
 	UPROPERTY()
 	TArray<int32> RowIndexes;
 };
@@ -51,68 +51,75 @@ class HOUDININIAGARA_API UHoudiniCSV : public UObject
 	//  MEMBER FUNCTIONS
 	//-----------------------------------------------------------------------------------------
 
+	virtual void GetAssetRegistryTags(TArray< FAssetRegistryTag > & OutTags) const override;
+
 	bool UpdateFromFile( const FString& TheFileName );
 	bool UpdateFromStringArray( TArray<FString>& StringArray );
 
 	void SetFileName( const FString& TheFilename );
 
 	// Returns the number of points found in the CSV file
-	int32 GetNumberOfParticlesInCSV();
+	int32 GetNumberOfPoints();
 
-	// Returns the number of points found in the CSV file
-	int32 GetNumberOfLinesInCSV();
+	// Returns the number of rows found in the CSV file
+	int32 GetNumberOfRows();
 
 	// Returns the number of columns found in the CSV file
-	int32 GetNumberOfColumnsInCSV();
+	int32 GetNumberOfColumns();
 
 	// Returns the column index for a given string
 	bool GetColumnIndexFromString(const FString& ColumnTitle, int32& ColumnIndex);
 
 	// Returns the float value at a given point in the CSV file
-	bool GetCSVFloatValue( const int32& lineIndex, const int32& colIndex, float& value );
+	bool GetFloatValue( const int32& rowIndex, const int32& colIndex, float& value );
 	// Returns the float value at a given point in the CSV file
-	bool GetCSVFloatValue( const int32& lineIndex, const FString& ColumnTitle, float& value );
+	bool GetFloatValue( const int32& rowIndex, const FString& ColumnTitle, float& value );
+	/*
 	// Returns the float value at a given point in the CSV file
-	bool GetCSVStringValue( const int32& lineIndex, const int32& colIndex, FString& value );
+	bool GetCSVStringValue( const int32& rowIndex, const int32& colIndex, FString& value );
 	// Returns the string value at a given point in the CSV file
-	bool GetCSVStringValue( const int32& lineIndex, const FString& ColumnTitle, FString& value );
+	bool GetCSVStringValue( const int32& rowIndex, const FString& ColumnTitle, FString& value );
+	*/
 	// Returns a Vector3 for a given point in the CSV file
-	bool GetCSVVectorValue( const int32& lineIndex, const int32& colIndex, FVector& value, const bool& DoSwap = true, const bool& DoScale = true );
+	bool GetVectorValue( const int32& rowIndex, const int32& colIndex, FVector& value, const bool& DoSwap = true, const bool& DoScale = true );
 	
 	// Returns a time value for a given point in the CSV file
-	bool GetCSVTimeValue( const int32& lineIndex, float& value );
+	bool GetTimeValue( const int32& rowIndex, float& value );
 	// Returns a Position Vector3 for a given point in the CSV file (converted to unreal's coordinate system)
-	bool GetCSVPositionValue( const int32& lineIndex, FVector& value );
+	bool GetPositionValue( const int32& rowIndex, FVector& value );
 	// Returns a Normal Vector3 for a given point in the CSV file (converted to unreal's coordinate system)
-	bool GetCSVNormalValue( const int32& lineIndex, FVector& value );
+	bool GetNormalValue( const int32& rowIndex, FVector& value );
 	// Returns a Color for a given point in the CSV file
-	bool GetCSVColorValue( const int32& lineIndex, FLinearColor& value );
+	bool GetColorValue( const int32& rowIndex, FLinearColor& value );
 	// Returns a Velocity Vector3 for a given point in the CSV file
-	bool GetCSVVelocityValue(const int32& lineIndex, FVector& value );
+	bool GetVelocityValue(const int32& rowIndex, FVector& value );
 
 	// Get the last row index for a given time value (the row with a time smaller or equal to desiredTime)
-	// If the CSV file doesn't have time informations, returns false and set the LastRowIndex to the last line in the file
+	// If the CSV file doesn't have time informations, returns false and set LastRowIndex to the last row in the file
 	// If desiredTime is smaller than the time value in the first row, LastRowIndex will be set to -1
 	// If desiredTime is higher than the last time value in the last row of the csv file, LastIndex will be set to the last row's index
 	bool GetLastRowIndexAtTime( const float& desiredTime, int32& lastRowIndex );
 
-	// Get the last index of the particles to be spawned at time t
+	// Get the last pointID of the points to be spawned at time t
 	// Invalid Index are used to indicate edge cases:
-	// -1 will be returned if no particles have been spawned ( t is smaller than the first particle time )
-	// NumberOfLines will be returned if all particles in the CSV have been spawned ( t is higher than the last particle time )
-	bool GetLastParticleIndexAtTime( const float& time, int32& lastIndex );
-
-	// Returns the previous and next indexes for reading the values of a specified particle at a given time
-	bool GetParticleLineIndexAtTime(const int32& ParticleID, const float& desiredTime, int32& PrevIndex, int32& NextIndex, float& PrevWeight);
-	// Returns the value for a particle at a given time value (linearly interpolated) 
-	bool GetParticleValueAtTime(const int32& ParticleID, const int32& ColumnIndex, const float& desiredTime, float& Value);
-	// Returns the Vector Value for a given particle at a given time value (linearly interpolated) 
-	bool GetParticleVectorValueAtTime(const int32& ParticleID, const int32& ColumnIndex, const float& desiredTime, FVector& Vector, const bool& DoSwap, const bool& DoScale);
-	// Returns the Position Value for a given particle at a given time value (linearly interpolated) 
-	bool GetParticlePositionAtTime(const int32& ParticleID, const float& desiredTime, FVector& Vector);
-
-	// Return a Particle's life for a given time value
-	bool GetParticleLifeAtTime(const int32& ParticleID, const float& DesiredTime, float& Value);
+	// -1 will be returned if there is no points to spawn ( t is smaller than the first point time )
+	// NumberOfRows will be returned if all points in the CSV have been spawned ( t is higher than the last point time )
+	bool GetLastPointIDToSpawnAtTime( const float& time, int32& lastIndex );
+
+	// Returns the previous and next indexes for reading the values of a specified point at a given time
+	bool GetRowIndexesForPointAtTime(const int32& PointID, const float& desiredTime, int32& PrevIndex, int32& NextIndex, float& PrevWeight);
+	// Returns the value for a point at a given time value (linearly interpolated) 
+	bool GetPointValueAtTime(const int32& PointID, const int32& ColumnIndex, const float& desiredTime, float& Value);
+	// Returns the Vector Value for a given point at a given time value (linearly interpolated) 
+	bool GetPointVectorValueAtTime(const int32& PointID, const int32& ColumnIndex, const float& desiredTime, FVector& Vector, const bool& DoSwap, const bool& DoScale);
+	// Returns the Position Value for a given point at a given time value (linearly interpolated) 
+	bool GetPointPositionAtTime(const int32& PointID, const float& desiredTime, FVector& Vector);
+	// Return a given point's life value at spawn
+	bool GetPointLife(const int32& PointID, float& Value);
+	// Return a point's life for a given time value
+	bool GetPointLifeAtTime(const int32& PointID, const float& DesiredTime, float& Value);
+	// Return a point's type at spawn
+	bool GetPointType(const int32& PointID, int32& Value);
 
 	//-----------------------------------------------------------------------------------------
 	//  MEMBER VARIABLES
@@ -122,57 +129,70 @@ class HOUDININIAGARA_API UHoudiniCSV : public UObject
 
 	// The number of values stored in the CSV file (excluding the title row)
 	UPROPERTY( VisibleAnywhere, Category = "Houdini CSV File Properties" )
-	int32 NumberOfLines;
+	int32 NumberOfRows;
 
 	// The number of value TYPES stored in the CSV file
 	UPROPERTY( VisibleAnywhere, Category = "Houdini CSV File Properties" )
 	int32 NumberOfColumns;
 
-	// The number of particles found in the CSV file
-	UPROPERTY(VisibleAnywhere, Category = "Houdini CSV File Properties")
-	int32 NumberOfParticles;
+	// The number of points found in the CSV file
+	UPROPERTY( VisibleAnywhere, Category = "Houdini CSV File Properties")
+	int32 NumberOfPoints;
 
-	// The tokenized title raw, describing the content of each column
+	// The source title row, describing the content of each column and used to locate specific values in the CSV file.
+	// Editing this will trigger a re-import of the CSV file.
+	UPROPERTY(EditAnywhere, Category = "Houdini CSV File Properties")
+	FString SourceTitleRow;
+
+	// The final column titles used by the asset after parsing.
+	// Describes the content of each column.
+	// Packed vector values are expended, so additional column might have been inserted.
+	// Use the indexes in this array to access your data.
 	UPROPERTY( VisibleAnywhere, Category = "Houdini CSV File Properties" )
-	TArray<FString> TitleRowArray;
+	TArray<FString> ColumnTitleArray;
 
 #if WITH_EDITORONLY_DATA
 	/** Importing data and options used for this asset */
-	UPROPERTY(EditAnywhere, Instanced, Category = ImportSettings)
+	UPROPERTY( EditAnywhere, Instanced, Category = ImportSettings )
 	class UAssetImportData* AssetImportData;
 
 	virtual void PostInitProperties() override;
+
+	virtual void PostEditChangeProperty( FPropertyChangedEvent & PropertyChangedEvent ) override;
 #endif
 
 	protected:
 
-	// Parses the CSV file's title line to update the column indexes of special values we're interested in
-	// Also look for packed vectors in the first line and update the indexes accordingly
-	bool ParseCSVTitleLine( const FString& CSVTitleString, const FString& FirstLine, bool& HasPackedVectors );
-
+	// Parses the CSV title row to update the column indexes of special values we're interested in
+	// Also look for packed vectors in the first row and update the indexes accordingly
+	bool ParseCSVTitleRow(const FString& TitleRow, const FString& FirstValueRow, bool& HasPackedVectors);
 
     private:
 
 	static bool SortPredicate(const TArray<FString>& A, const TArray<FString>& B);
 
 	// Array containing the Raw String data
-	UPROPERTY()
-	TArray<FString> StringCSVData;
+	//UPROPERTY()
+	//TArray<FString> StringCSVData;
 
 	// Array containing all the CSV data converted to floats
 	UPROPERTY()
 	TArray<float> FloatCSVData;
 	
-	// Array containing the spawn times for each particles in the file
+	// Array containing the spawn times for each point in the file
 	UPROPERTY()
 	TArray<float> SpawnTimes;
 
-	// Array containing all the life values for each particles in the file
+	// Array containing all the life values for each point in the file
 	UPROPERTY()
 	TArray<float> LifeValues;
 
+	// Array containing all the type values for each point in the file
+	UPROPERTY()
+	TArray<int32> PointTypes;
+
 	// Index of the Position values in the buffer
-	UPROPERTY()	    
+	UPROPERTY()
 	int32 PositionColumnIndex;
 
 	// Index of the Normal values in the buffer
@@ -183,7 +203,7 @@ class HOUDININIAGARA_API UHoudiniCSV : public UObject
 	UPROPERTY()
 	int32 TimeColumnIndex;
 
-	// Index of the particle id values in the buffer
+	// Index of the point id values in the buffer
 	UPROPERTY()
 	int32 IDColumnIndex;
 
@@ -207,11 +227,18 @@ class HOUDININIAGARA_API UHoudiniCSV : public UObject
 	UPROPERTY()
 	int32 VelocityColumnIndex;
 
+	// Column index of the "type" values in the buffer
+	UPROPERTY()
+	int32 TypeColumnIndex;
+
 	// Row indexes for new time values
 	//UPROPERTY()
 	//TMap<float, int32> TimeValuesIndexes;
 
-	// Row indexes for each particle
+	// Row indexes for each point
+	UPROPERTY()
+	TArray< FPointIndexes > PointValueIndexes;
+
 	UPROPERTY()
-	TArray< FParticleIndexes > ParticleValueIndexes;
+	bool UseCustomTitleRow;
 };
\ No newline at end of file
diff --git a/Engine/Plugins/FX/HoudiniNiagara/Source/HoudiniNiagaraEditor/Private/HoudiniCSVAssetActions.cpp b/Engine/Plugins/FX/HoudiniNiagara/Source/HoudiniNiagaraEditor/Private/HoudiniCSVAssetActions.cpp
--- a/Engine/Plugins/FX/HoudiniNiagara/Source/HoudiniNiagaraEditor/Private/HoudiniCSVAssetActions.cpp
+++ b/Engine/Plugins/FX/HoudiniNiagara/Source/HoudiniNiagaraEditor/Private/HoudiniCSVAssetActions.cpp
@@ -27,6 +27,7 @@
 #include "EditorStyleSet.h"
 #include "Toolkits/AssetEditorToolkit.h"
 #include "EditorReimportHandler.h"
+#include "HAL/FileManager.h"
 
 //#include "HoudiniCSVAssetEditorToolkit.h"
 
@@ -67,6 +68,16 @@ void FHoudiniCSVAssetActions::GetActions(const TArray<UObject*>& InObjects, FMen
 			FCanExecuteAction::CreateSP(this, &FHoudiniCSVAssetActions::CanExecuteOpenInEditor, HoudiniCSVAssets)
 		)
     );
+
+	MenuBuilder.AddMenuEntry(
+		LOCTEXT("FindHoudiniCSVInExplorer", "Find Source CSV"),
+		LOCTEXT("FindHoudiniCSVInExplorerTooltip", "Opens explorer at the location of this asset's source CSV file."),
+		FSlateIcon(FEditorStyle::GetStyleSetName(), "ContentBrowser.AssetActions.OpenInExternalEditor"),
+		FUIAction(
+			FExecuteAction::CreateSP(this, &FHoudiniCSVAssetActions::ExecuteFindInExplorer, HoudiniCSVAssets),
+			FCanExecuteAction::CreateSP(this, &FHoudiniCSVAssetActions::CanExecuteFindInExplorer, HoudiniCSVAssets)
+		)
+	);
 }
 
 
@@ -89,13 +100,26 @@ UClass* FHoudiniCSVAssetActions::GetSupportedClass() const
 
 FText FHoudiniCSVAssetActions::GetAssetDescription(const FAssetData& AssetData) const
 {
-    return FText::GetEmpty();
+	/*
+	UHoudiniCSV* CSVAsset = Cast<UHoudiniCSV>(AssetData.GetAsset());
+	if ( !CSVAsset )
+		return FText::GetEmpty();
+
+	FString StrDescription;
+	StrDescription += TEXT("Number of Rows: ") + FString::FromInt(CSVAsset->NumberOfRows) + TEXT("\n");
+	StrDescription += TEXT("Number of Columns: ") + FString::FromInt(CSVAsset->NumberOfColumns) + TEXT("\n");
+	StrDescription += TEXT("Number of Points: ") + FString::FromInt(CSVAsset->NumberOfPoints) + TEXT("\n");
+
+	return FText::FromString(StrDescription);
+	*/
+	return FText::GetEmpty();
 }
 
 
 FColor FHoudiniCSVAssetActions::GetTypeColor() const
 {
-    return FColor::Orange;
+    //return FColor::Orange;
+	return FColor(255, 165, 0);
 }
 
 
@@ -125,14 +149,17 @@ class UThumbnailInfo* FHoudiniCSVAssetActions::GetThumbnailInfo(UObject* Asset)
 
 bool FHoudiniCSVAssetActions::CanExecuteReimport(const TArray<TWeakObjectPtr<UHoudiniCSV>> Objects) const
 {
-    for (auto ObjIt = Objects.CreateConstIterator(); ObjIt; ++ObjIt)
+    for ( auto ObjIt = Objects.CreateConstIterator(); ObjIt; ++ObjIt )
     {
 		auto Object = (*ObjIt).Get();
 		if (!Object)
-			return false;
+			continue;
+
+		// If one object is valid, we can execute the action
+		return true;
     }
 
-    return true;
+    return false;
 }
 
 void FHoudiniCSVAssetActions::ExecuteReimport(const TArray<TWeakObjectPtr<UHoudiniCSV>> Objects) const
@@ -140,11 +167,10 @@ void FHoudiniCSVAssetActions::ExecuteReimport(const TArray<TWeakObjectPtr<UHoudi
     for ( auto ObjIt = Objects.CreateConstIterator(); ObjIt; ++ObjIt )
     {
 		auto Object = (*ObjIt).Get();
-		if ( Object )
-		{
-			// Fonts fail to reimport if they ask for a new file if missing
-			FReimportManager::Instance()->Reimport(Object, true);
-		}
+		if (!Object)
+			continue;
+
+		FReimportManager::Instance()->Reimport(Object, true);
     }
 }
 
@@ -155,11 +181,18 @@ FHoudiniCSVAssetActions::CanExecuteOpenInEditor(const TArray<TWeakObjectPtr<UHou
     for (auto ObjIt = Objects.CreateConstIterator(); ObjIt; ++ObjIt)
     {
 		auto Object = ( *ObjIt ).Get();
-		if ( !Object )
-			return false;
+		if (!Object)
+			continue;
+
+		const FString SourceFilePath = Object->FileName;
+		if (!SourceFilePath.Len() || IFileManager::Get().FileSize(*SourceFilePath) == INDEX_NONE)
+			continue;
+
+		// If one object is valid, we can execute the action
+		return true;
     }
 
-    return true;
+    return false;
 }
 
 void
@@ -168,12 +201,50 @@ FHoudiniCSVAssetActions::ExecuteOpenInEditor(const TArray<TWeakObjectPtr<UHoudin
     for ( auto ObjIt = Objects.CreateConstIterator(); ObjIt; ++ObjIt )
     {
 		auto Object = ( *ObjIt ).Get();
-		if ( Object )
-		{
-			// Fonts fail to reimport if they ask for a new file if missing
-			FPlatformProcess::LaunchFileInDefaultExternalApplication(*(Object->FileName), NULL, ELaunchVerb::Open);
-		}
+		if (!Object)
+			continue;
+
+		const FString SourceFilePath = Object->FileName;
+		if (!SourceFilePath.Len() || IFileManager::Get().FileSize(*SourceFilePath) == INDEX_NONE)
+			continue;
+
+		FPlatformProcess::LaunchFileInDefaultExternalApplication(*SourceFilePath, NULL, ELaunchVerb::Open);
     }
 }
 
+bool
+FHoudiniCSVAssetActions::CanExecuteFindInExplorer(const TArray<TWeakObjectPtr<UHoudiniCSV>> Objects) const
+{
+	for ( auto ObjIt = Objects.CreateConstIterator(); ObjIt; ++ObjIt )
+	{
+		auto Object = (*ObjIt).Get();
+		if (!Object)
+			continue;
+
+		const FString SourceFilePath = Object->FileName;
+		if (!SourceFilePath.Len() || IFileManager::Get().FileSize(*SourceFilePath) == INDEX_NONE)
+			continue;
+
+		// If one object is valid, we can execute the action
+		return true;
+	}
+
+	return false;
+}
+
+void
+FHoudiniCSVAssetActions::ExecuteFindInExplorer(const TArray<TWeakObjectPtr<UHoudiniCSV>> Objects) const
+{
+	for ( auto ObjIt = Objects.CreateConstIterator(); ObjIt; ++ObjIt )
+	{
+		auto Object = (*ObjIt).Get();
+		if (!Object)
+			continue;
+
+		const FString SourceFilePath = Object->FileName;
+		if ( SourceFilePath.Len() && IFileManager::Get().FileSize(*SourceFilePath) != INDEX_NONE )
+			return FPlatformProcess::ExploreFolder(*SourceFilePath);
+	}
+}
+
 #undef LOCTEXT_NAMESPACE
\ No newline at end of file
diff --git a/Engine/Plugins/FX/HoudiniNiagara/Source/HoudiniNiagaraEditor/Public/HoudiniCSVAssetActions.h b/Engine/Plugins/FX/HoudiniNiagara/Source/HoudiniNiagaraEditor/Public/HoudiniCSVAssetActions.h
--- a/Engine/Plugins/FX/HoudiniNiagara/Source/HoudiniNiagaraEditor/Public/HoudiniCSVAssetActions.h
+++ b/Engine/Plugins/FX/HoudiniNiagara/Source/HoudiniNiagaraEditor/Public/HoudiniCSVAssetActions.h
@@ -58,4 +58,8 @@ class FHoudiniCSVAssetActions : public FAssetTypeActions_Base
 	// Open In Editor Actions
 	bool CanExecuteOpenInEditor(const TArray<TWeakObjectPtr<UHoudiniCSV>> Objects) const;
 	void ExecuteOpenInEditor(const TArray<TWeakObjectPtr<UHoudiniCSV>> Objects) const;
+
+	// Find In Explorer Actions
+	bool CanExecuteFindInExplorer(const TArray<TWeakObjectPtr<UHoudiniCSV>> Objects) const;
+	void ExecuteFindInExplorer(const TArray<TWeakObjectPtr<UHoudiniCSV>> Objects) const;	
 };
\ No newline at end of file
