Fixing nightly build due to missing GetAssetTags definition due to mismatches in WITH_EDITORONLY_DATA

#tests n/a

diff --git a/Engine/Plugins/FX/HoudiniNiagara/Source/HoudiniNiagara/Public/HoudiniCSV.h b/Engine/Plugins/FX/HoudiniNiagara/Source/HoudiniNiagara/Public/HoudiniCSV.h
--- a/Engine/Plugins/FX/HoudiniNiagara/Source/HoudiniNiagara/Public/HoudiniCSV.h
+++ b/Engine/Plugins/FX/HoudiniNiagara/Source/HoudiniNiagara/Public/HoudiniCSV.h
@@ -50,9 +50,7 @@ class HOUDININIAGARA_API UHoudiniCSV : public UObject
 	//-----------------------------------------------------------------------------------------
 	//  MEMBER FUNCTIONS
 	//-----------------------------------------------------------------------------------------
-
-	virtual void GetAssetRegistryTags(TArray< FAssetRegistryTag > & OutTags) const override;
-
+		
 	bool UpdateFromFile( const FString& TheFileName );
 	bool UpdateFromStringArray( TArray<FString>& StringArray );
 
@@ -158,7 +156,9 @@ class HOUDININIAGARA_API UHoudiniCSV : public UObject
 
 	virtual void PostInitProperties() override;
 
-	virtual void PostEditChangeProperty( FPropertyChangedEvent & PropertyChangedEvent ) override;
+	virtual void PostEditChangeProperty(FPropertyChangedEvent & PropertyChangedEvent) override;
+
+	virtual void GetAssetRegistryTags(TArray< FAssetRegistryTag > & OutTags) const override;
 #endif
 
 	protected:
