Houdini-Niagara:
Houdini CSV Asset:
- Fixed build break due to GetAssetRegistryTags()
- Replaced the different hardcoded ColumnIndexes member variables by an array.

diff --git a/Engine/Plugins/FX/HoudiniNiagara/Source/HoudiniNiagara/Private/HoudiniCSV.cpp b/Engine/Plugins/FX/HoudiniNiagara/Source/HoudiniNiagara/Private/HoudiniCSV.cpp
--- a/Engine/Plugins/FX/HoudiniNiagara/Source/HoudiniNiagara/Private/HoudiniCSV.cpp
+++ b/Engine/Plugins/FX/HoudiniNiagara/Source/HoudiniNiagara/Private/HoudiniCSV.cpp
@@ -79,16 +79,7 @@ UHoudiniCSV::UHoudiniCSV( const FObjectInitializer& ObjectInitializer )
 	NumberOfColumns( -1 ),
 	NumberOfPoints( -1 )
 {
-	PositionColumnIndex = INDEX_NONE;
-	NormalColumnIndex = INDEX_NONE;
-	TimeColumnIndex = INDEX_NONE;
-	IDColumnIndex = INDEX_NONE;
-	AliveColumnIndex = INDEX_NONE;
-	LifeColumnIndex = INDEX_NONE;
-	ColorColumnIndex = INDEX_NONE;
-	AlphaColumnIndex = INDEX_NONE;
-	VelocityColumnIndex = INDEX_NONE;
-	TypeColumnIndex = INDEX_NONE;
+	SpecialAttributesColumnIndexes.Init(INDEX_NONE, EHoudiniAttributes::HOUDINI_ATTR_SIZE);
 
 	UseCustomTitleRow = false;
 }
@@ -125,17 +116,8 @@ bool UHoudiniCSV::UpdateFromStringArray( TArray<FString>& RawStringArray )
     NumberOfRows = 0;
 	NumberOfPoints = 0;
 
-    // Reset the position, normal and time indexes
-    PositionColumnIndex = INDEX_NONE;
-    NormalColumnIndex = INDEX_NONE;
-    TimeColumnIndex = INDEX_NONE;
-	IDColumnIndex = INDEX_NONE;
-	AliveColumnIndex = INDEX_NONE;
-	LifeColumnIndex = INDEX_NONE;
-	ColorColumnIndex = INDEX_NONE;
-	AlphaColumnIndex = INDEX_NONE;
-	VelocityColumnIndex = INDEX_NONE;
-	TypeColumnIndex = INDEX_NONE;
+	// Reset the column indexes of the special attributes
+	SpecialAttributesColumnIndexes.Init( INDEX_NONE, EHoudiniAttributes::HOUDINI_ATTR_SIZE );
 
     if ( RawStringArray.Num() <= 0 )
     {
@@ -203,6 +185,8 @@ bool UHoudiniCSV::UpdateFromStringArray( TArray<FString>& RawStringArray )
 	}
 
 	// If we have time values, we have to make sure the csv rows are sorted by time
+	int32 TimeColumnIndex = GetAttributeColumnIndex(EHoudiniAttributes::TIME);
+	int32 IDColumnIndex = GetAttributeColumnIndex(EHoudiniAttributes::POINTID);
 	if ( TimeColumnIndex != INDEX_NONE )
 	{
 		// First check if we need to sort the array
@@ -335,7 +319,11 @@ bool UHoudiniCSV::UpdateFromStringArray( TArray<FString>& RawStringArray )
 	if ( NumberOfPoints <= 0 )
 		NumberOfPoints = NumberOfRows;
 
-	// Look for point specific attributes
+	// Look for point specific attributes to build some helper arrays
+	int32 LifeColumnIndex = GetAttributeColumnIndex(EHoudiniAttributes::LIFE);
+	int32 AliveColumnIndex = GetAttributeColumnIndex(EHoudiniAttributes::ALIVE);
+	int32 TypeColumnIndex = GetAttributeColumnIndex(EHoudiniAttributes::TYPE);
+
 	SpawnTimes.Empty();
 	SpawnTimes.Init( -1.0f, NumberOfPoints );
 
@@ -344,7 +332,7 @@ bool UHoudiniCSV::UpdateFromStringArray( TArray<FString>& RawStringArray )
 
 	PointTypes.Empty();
 	PointTypes.Init( -1, NumberOfPoints );
-
+	
 	float MaxTime = -1.0f;
 	for ( int rowIdx = 0; rowIdx < NumberOfRows; rowIdx++ )
 	{
@@ -428,62 +416,74 @@ bool UHoudiniCSV::ParseCSVTitleRow( const FString& TitleRow, const FString& Firs
 			|| CurrentTitle.Equals( TEXT("X"), ESearchCase::IgnoreCase )
 			|| CurrentTitle.Equals(TEXT("pos"), ESearchCase::IgnoreCase ) )
 		{
-			if ( PositionColumnIndex == INDEX_NONE )
-				PositionColumnIndex = n;
+			if ( !IsValidAttributeColumnIndex(EHoudiniAttributes::POSITION))
+				SpecialAttributesColumnIndexes[EHoudiniAttributes::POSITION] = n;
 		}
 		else if ( CurrentTitle.Equals( TEXT("N"), ESearchCase::IgnoreCase )
 			|| CurrentTitle.Equals( TEXT("Nx"), ESearchCase::IgnoreCase ) )
 		{
-			if ( NormalColumnIndex == INDEX_NONE )
-				NormalColumnIndex = n;
+			if (!IsValidAttributeColumnIndex(EHoudiniAttributes::NORMAL))
+				SpecialAttributesColumnIndexes[EHoudiniAttributes::NORMAL] = n;
 		}
 		else if ( ( CurrentTitle.Equals( TEXT("T"), ESearchCase::IgnoreCase ) )
 			|| ( CurrentTitle.Contains( TEXT("time"), ESearchCase::IgnoreCase ) ) )
 		{
-			if ( TimeColumnIndex == INDEX_NONE )
-				TimeColumnIndex = n;
+			if (!IsValidAttributeColumnIndex(EHoudiniAttributes::TIME))
+				SpecialAttributesColumnIndexes[EHoudiniAttributes::TIME] = n;
 		}
 		else if ( ( CurrentTitle.Equals( TEXT("#"), ESearchCase::IgnoreCase ) )
 			|| ( CurrentTitle.Equals( TEXT("id"), ESearchCase::IgnoreCase ) ) )
 		{
-			if ( IDColumnIndex == INDEX_NONE )
-				IDColumnIndex = n;
+			if (!IsValidAttributeColumnIndex(EHoudiniAttributes::POINTID))
+				SpecialAttributesColumnIndexes[EHoudiniAttributes::POINTID] = n;
 		}
 		else if ( CurrentTitle.Equals( TEXT("alive"), ESearchCase::IgnoreCase ) )
 		{
-			if ( AliveColumnIndex == INDEX_NONE )
-				AliveColumnIndex = n;
+			if (!IsValidAttributeColumnIndex(EHoudiniAttributes::ALIVE))
+				SpecialAttributesColumnIndexes[EHoudiniAttributes::ALIVE] = n;
 		}
 		else if ( CurrentTitle.Equals( TEXT("life"), ESearchCase::IgnoreCase ) )
 		{
-			if ( LifeColumnIndex == INDEX_NONE)
-				LifeColumnIndex = n;
+			if (!IsValidAttributeColumnIndex(EHoudiniAttributes::LIFE))
+				SpecialAttributesColumnIndexes[EHoudiniAttributes::LIFE] = n;
 		}
 		else if ( ( CurrentTitle.Equals( TEXT("Cd"), ESearchCase::IgnoreCase ) )
 			|| ( CurrentTitle.Equals(TEXT("color"), ESearchCase::IgnoreCase ) ) )
 		{
-			if ( ColorColumnIndex == INDEX_NONE )
-				ColorColumnIndex = n;
+			if (!IsValidAttributeColumnIndex(EHoudiniAttributes::COLOR))
+				SpecialAttributesColumnIndexes[EHoudiniAttributes::COLOR] = n;
 		}
 		else if ( ( CurrentTitle.Equals( TEXT("alpha"), ESearchCase::IgnoreCase ) )
 			|| ( CurrentTitle.Equals(TEXT("A"), ESearchCase::IgnoreCase ) ) )
 		{
-			if ( AlphaColumnIndex == INDEX_NONE )
-				AlphaColumnIndex = n;
+			if (!IsValidAttributeColumnIndex(EHoudiniAttributes::ALPHA))
+				SpecialAttributesColumnIndexes[EHoudiniAttributes::ALPHA] = n;
 		}
 		else if ( ( CurrentTitle.Equals( TEXT("v"), ESearchCase::IgnoreCase ) )
 			|| ( CurrentTitle.Equals(TEXT("Vx"), ESearchCase::IgnoreCase ) ) )
 		{
-			if ( VelocityColumnIndex == INDEX_NONE )
-				VelocityColumnIndex = n;
+			if (!IsValidAttributeColumnIndex(EHoudiniAttributes::VELOCITY))
+				SpecialAttributesColumnIndexes[EHoudiniAttributes::VELOCITY] = n;
 		}
 		else if ( CurrentTitle.Equals(TEXT("type"), ESearchCase::IgnoreCase ) )
 		{
-			if ( TypeColumnIndex == INDEX_NONE )
-				TypeColumnIndex = n;
+			if (!IsValidAttributeColumnIndex(EHoudiniAttributes::TYPE))
+				SpecialAttributesColumnIndexes[EHoudiniAttributes::TYPE] = n;
 		}
     }
 
+	// Extract the special attributes column indexes we found
+	int32 PositionColumnIndex = GetAttributeColumnIndex(EHoudiniAttributes::POSITION);
+	int32 NormalColumnIndex = GetAttributeColumnIndex(EHoudiniAttributes::NORMAL);
+	int32 TimeColumnIndex = GetAttributeColumnIndex(EHoudiniAttributes::TIME);
+	int32 IDColumnIndex = GetAttributeColumnIndex(EHoudiniAttributes::POINTID);
+	int32 AliveColumnIndex = GetAttributeColumnIndex(EHoudiniAttributes::ALIVE);
+	int32 LifeColumnIndex = GetAttributeColumnIndex(EHoudiniAttributes::LIFE);
+	int32 ColorColumnIndex = GetAttributeColumnIndex(EHoudiniAttributes::COLOR);
+	int32 AlphaColumnIndex = GetAttributeColumnIndex(EHoudiniAttributes::ALPHA);
+	int32 VelocityColumnIndex = GetAttributeColumnIndex(EHoudiniAttributes::VELOCITY);
+	int32 TypeColumnIndex = GetAttributeColumnIndex(EHoudiniAttributes::TYPE);
+
 	// Read the first row of the CSV file, and look for packed vectors value (X,Y,Z)
     // We'll have to expand them in the title row to match the parsed data
 	HasPackedVectors = false;
@@ -578,36 +578,16 @@ bool UHoudiniCSV::ParseCSVTitleRow( const FString& TitleRow, const FString& Firs
 			}
 		}
 
-		// Eventually offset the stored index
-		if ( PositionColumnIndex != INDEX_NONE && ( PositionColumnIndex > FoundPackedVectorColumnIndex) )
-			PositionColumnIndex += FoundVectorSize - 1;
-
-		if ( NormalColumnIndex != INDEX_NONE && ( NormalColumnIndex > FoundPackedVectorColumnIndex) )
-			NormalColumnIndex += FoundVectorSize - 1;
-
-		if ( TimeColumnIndex != INDEX_NONE && ( TimeColumnIndex > FoundPackedVectorColumnIndex) )
-			TimeColumnIndex += FoundVectorSize - 1;
-
-		if ( IDColumnIndex != INDEX_NONE && ( IDColumnIndex > FoundPackedVectorColumnIndex ) )
-			IDColumnIndex += FoundVectorSize - 1;
-
-		if ( AliveColumnIndex != INDEX_NONE && ( AliveColumnIndex > FoundPackedVectorColumnIndex ) )
-			AliveColumnIndex += FoundVectorSize - 1;
-
-		if ( LifeColumnIndex != INDEX_NONE && ( LifeColumnIndex > FoundPackedVectorColumnIndex ) )
-			LifeColumnIndex += FoundVectorSize - 1;
-
-		if ( ColorColumnIndex != INDEX_NONE && ( ColorColumnIndex > FoundPackedVectorColumnIndex ) )
-			ColorColumnIndex += FoundVectorSize - 1;
-
-		if ( AlphaColumnIndex != INDEX_NONE && ( AlphaColumnIndex > FoundPackedVectorColumnIndex ) )
-			AlphaColumnIndex += FoundVectorSize - 1;
-
-		if ( VelocityColumnIndex != INDEX_NONE && ( VelocityColumnIndex > FoundPackedVectorColumnIndex ) )
-			VelocityColumnIndex += FoundVectorSize - 1;
+		// We have inserted new column titles because of the packed vector,
+		// so we need to offset all the special attributes column indexes after the inserted values
+		for ( int32 AttrIdx = EHoudiniAttributes::HOUDINI_ATTR_BEGIN; AttrIdx < EHoudiniAttributes::HOUDINI_ATTR_SIZE; AttrIdx++ )
+		{
+			if ( SpecialAttributesColumnIndexes[AttrIdx] == INDEX_NONE )
+				continue;
 
-		if ( TypeColumnIndex != INDEX_NONE && ( TypeColumnIndex > FoundPackedVectorColumnIndex ) )
-			TypeColumnIndex += FoundVectorSize - 1;
+			if ( SpecialAttributesColumnIndexes[AttrIdx] > FoundPackedVectorColumnIndex )
+				SpecialAttributesColumnIndexes[AttrIdx] += FoundVectorSize - 1;
+		}
 
 		HasPackedVectors = true;
 		FoundPackedVectorCharIndex++;
@@ -645,7 +625,7 @@ bool UHoudiniCSV::ParseCSVTitleRow( const FString& TitleRow, const FString& Firs
 
 
 // Returns the float value at a given point in the CSV file
-bool UHoudiniCSV::GetFloatValue( const int32& rowIndex, const int32& colIndex, float& value )
+bool UHoudiniCSV::GetFloatValue( const int32& rowIndex, const int32& colIndex, float& value ) const
 {
     if ( rowIndex < 0 || rowIndex >= NumberOfRows )
 		return false;
@@ -685,7 +665,7 @@ bool UHoudiniCSV::GetCSVStringValue( const int32& rowIndex, const int32& colInde
 */
 
 // Returns a Vector 3 for a given point in the CSV file
-bool UHoudiniCSV::GetVectorValue( const int32& rowIndex, const int32& colIndex, FVector& value, const bool& DoSwap, const bool& DoScale )
+bool UHoudiniCSV::GetVectorValue( const int32& rowIndex, const int32& colIndex, FVector& value, const bool& DoSwap, const bool& DoScale ) const
 {
     FVector V = FVector::ZeroVector;
     if ( !GetFloatValue( rowIndex, colIndex, V.X ) )
@@ -712,10 +692,10 @@ bool UHoudiniCSV::GetVectorValue( const int32& rowIndex, const int32& colIndex,
 }
 
 // Returns a Vector 3 for a given point in the CSV file
-bool UHoudiniCSV::GetPositionValue( const int32& rowIndex, FVector& value )
+bool UHoudiniCSV::GetPositionValue( const int32& rowIndex, FVector& value ) const
 {
     FVector V = FVector::ZeroVector;
-    if ( !GetVectorValue( rowIndex, PositionColumnIndex, V, true, true ) )
+    if ( !GetVectorValue( rowIndex, GetAttributeColumnIndex(EHoudiniAttributes::POSITION), V, true, true ) )
 		return false;
 
     value = V;
@@ -724,10 +704,10 @@ bool UHoudiniCSV::GetPositionValue( const int32& rowIndex, FVector& value )
 }
 
 // Returns a Vector 3 for a given point in the CSV file
-bool UHoudiniCSV::GetNormalValue( const int32& rowIndex, FVector& value )
+bool UHoudiniCSV::GetNormalValue( const int32& rowIndex, FVector& value ) const
 {
     FVector V = FVector::ZeroVector;
-    if ( !GetVectorValue( rowIndex, NormalColumnIndex, V, true, false ) )
+    if ( !GetVectorValue( rowIndex, GetAttributeColumnIndex(EHoudiniAttributes::NORMAL), V, true, false ) )
 		return false;
 
     value = V;
@@ -736,10 +716,10 @@ bool UHoudiniCSV::GetNormalValue( const int32& rowIndex, FVector& value )
 }
 
 // Returns a time value for a given point in the CSV file
-bool UHoudiniCSV::GetTimeValue( const int32& rowIndex, float& value )
+bool UHoudiniCSV::GetTimeValue( const int32& rowIndex, float& value ) const
 {
     float temp;
-    if ( !GetFloatValue( rowIndex, TimeColumnIndex, temp ) )
+    if ( !GetFloatValue( rowIndex, GetAttributeColumnIndex(EHoudiniAttributes::TIME), temp ) )
 		return false;
 
     value = temp;
@@ -748,10 +728,10 @@ bool UHoudiniCSV::GetTimeValue( const int32& rowIndex, float& value )
 }
 
 // Returns a Color for a given point in the CSV file
-bool UHoudiniCSV::GetColorValue( const int32& rowIndex, FLinearColor& value )
+bool UHoudiniCSV::GetColorValue( const int32& rowIndex, FLinearColor& value ) const
 {
 	FVector V = FVector::OneVector;
-	if ( !GetVectorValue( rowIndex, ColorColumnIndex, V, false, false ) )
+	if ( !GetVectorValue( rowIndex, GetAttributeColumnIndex(EHoudiniAttributes::COLOR), V, false, false ) )
 		return false;
 
 	FLinearColor C = FLinearColor::White;
@@ -760,7 +740,7 @@ bool UHoudiniCSV::GetColorValue( const int32& rowIndex, FLinearColor& value )
 	C.B = V.Z;
 
 	float alpha = 1.0f;
-	if ( GetFloatValue( rowIndex, AlphaColumnIndex, alpha ) )
+	if ( GetFloatValue( rowIndex, GetAttributeColumnIndex(EHoudiniAttributes::ALPHA), alpha ) )
 		C.A = alpha;	
 
 	value = C;
@@ -769,10 +749,10 @@ bool UHoudiniCSV::GetColorValue( const int32& rowIndex, FLinearColor& value )
 }
 
 // Returns a Velocity Vector3 for a given point in the CSV file
-bool UHoudiniCSV::GetVelocityValue( const int32& rowIndex, FVector& value )
+bool UHoudiniCSV::GetVelocityValue( const int32& rowIndex, FVector& value ) const
 {
 	FVector V = FVector::ZeroVector;
-	if ( !GetVectorValue( rowIndex, VelocityColumnIndex, V, true, false ) )
+	if ( !GetVectorValue( rowIndex, GetAttributeColumnIndex(EHoudiniAttributes::VELOCITY), V, true, false ) )
 		return false;
 
 	value = V;
@@ -781,22 +761,22 @@ bool UHoudiniCSV::GetVelocityValue( const int32& rowIndex, FVector& value )
 }
 
 // Returns the number of points found in the CSV file
-int32 UHoudiniCSV::GetNumberOfPoints()
+int32 UHoudiniCSV::GetNumberOfPoints() const
 {
-	if ( IDColumnIndex != INDEX_NONE )
+	if ( !IsValidAttributeColumnIndex(EHoudiniAttributes::POINTID) )
 		return NumberOfPoints;
 
     return NumberOfRows;
 }
 
 // Returns the number of rows found in the CSV file
-int32 UHoudiniCSV::GetNumberOfRows()
+int32 UHoudiniCSV::GetNumberOfRows() const
 {
 	return NumberOfRows;
 }
 
 // Returns the number of columns found in the CSV file
-int32 UHoudiniCSV::GetNumberOfColumns()
+int32 UHoudiniCSV::GetNumberOfColumns() const
 {
 	return NumberOfColumns;
 }
@@ -805,9 +785,10 @@ int32 UHoudiniCSV::GetNumberOfColumns()
 // If the CSV file doesn't have time informations, returns false and set the LastRowIndex to the last row in the file
 // If desiredTime is smaller than the time value in the first row, LastRowIndex will be set to -1
 // If desiredTime is higher than the last time value in the last row of the csv file, LastIndex will be set to the last row's index
-bool UHoudiniCSV::GetLastRowIndexAtTime(const float& desiredTime, int32& lastRowIndex)
+bool UHoudiniCSV::GetLastRowIndexAtTime(const float& desiredTime, int32& lastRowIndex) const
 {
 	// If we dont have proper time info, always return the last row
+	int32 TimeColumnIndex = GetAttributeColumnIndex(EHoudiniAttributes::TIME);	
 	if ( TimeColumnIndex < 0 || TimeColumnIndex >= NumberOfColumns )
 	{
 		lastRowIndex = NumberOfRows - 1;
@@ -850,9 +831,10 @@ bool UHoudiniCSV::GetLastRowIndexAtTime(const float& desiredTime, int32& lastRow
 // If the CSV file doesn't have time informations, returns false and set the LastIndex to the last point
 // If desiredTime is smaller than the first point time, LastIndex will be set to -1
 // If desiredTime is higher than the last point time in the csv file, LastIndex will be set to the last point's index
-bool UHoudiniCSV::GetLastPointIDToSpawnAtTime( const float& desiredTime, int32& lastID )
+bool UHoudiniCSV::GetLastPointIDToSpawnAtTime( const float& desiredTime, int32& lastID ) const
 {
-    // If we dont have proper time info, always return the last point
+	// If we dont have proper time info, always return the last point
+	int32 TimeColumnIndex = GetAttributeColumnIndex(EHoudiniAttributes::TIME);    
     if ( TimeColumnIndex < 0 || TimeColumnIndex >= NumberOfColumns )
     {
 		lastID = NumberOfPoints - 1;
@@ -893,7 +875,7 @@ bool UHoudiniCSV::GetLastPointIDToSpawnAtTime( const float& desiredTime, int32&
 	return true;
 }
 
-bool UHoudiniCSV::GetPointType(const int32& PointID, int32& Value)
+bool UHoudiniCSV::GetPointType(const int32& PointID, int32& Value) const
 {
 	if ( !PointTypes.IsValidIndex( PointID ) )
 	{
@@ -906,7 +888,7 @@ bool UHoudiniCSV::GetPointType(const int32& PointID, int32& Value)
 	return true;
 }
 
-bool UHoudiniCSV::GetPointLife(const int32& PointID, float& Value)
+bool UHoudiniCSV::GetPointLife(const int32& PointID, float& Value) const
 {
 	if ( !LifeValues.IsValidIndex( PointID ) )
 	{
@@ -919,7 +901,7 @@ bool UHoudiniCSV::GetPointLife(const int32& PointID, float& Value)
 	return true;
 }
 
-bool UHoudiniCSV::GetPointLifeAtTime( const int32& PointID, const float& DesiredTime, float& Value )
+bool UHoudiniCSV::GetPointLifeAtTime( const int32& PointID, const float& DesiredTime, float& Value ) const
 {
 	if ( !SpawnTimes.IsValidIndex( PointID )  || !LifeValues.IsValidIndex( PointID ) )
 	{
@@ -939,7 +921,7 @@ bool UHoudiniCSV::GetPointLifeAtTime( const int32& PointID, const float& Desired
 	return true;
 }
 
-bool UHoudiniCSV::GetPointValueAtTime( const int32& PointID, const int32& ColumnIndex, const float& desiredTime, float& Value )
+bool UHoudiniCSV::GetPointValueAtTime( const int32& PointID, const int32& ColumnIndex, const float& desiredTime, float& Value ) const
 {
 	int32 PrevIndex = -1;
 	int32 NextIndex = -1;
@@ -959,10 +941,10 @@ bool UHoudiniCSV::GetPointValueAtTime( const int32& PointID, const int32& Column
 	return true;
 }
 
-bool UHoudiniCSV::GetPointPositionAtTime( const int32& PointID, const float& desiredTime, FVector& Vector )
+bool UHoudiniCSV::GetPointPositionAtTime( const int32& PointID, const float& desiredTime, FVector& Vector ) const
 {
 	FVector V = FVector::ZeroVector;
-	if ( !GetPointVectorValueAtTime(PointID, PositionColumnIndex, desiredTime, V, true, true ) )
+	if ( !GetPointVectorValueAtTime(PointID, GetAttributeColumnIndex(EHoudiniAttributes::POSITION), desiredTime, V, true, true ) )
 		return false;
 
 	Vector = V;
@@ -970,7 +952,7 @@ bool UHoudiniCSV::GetPointPositionAtTime( const int32& PointID, const float& des
 	return true;
 }
 
-bool UHoudiniCSV::GetPointVectorValueAtTime( const int32& PointID, const int32& ColumnIndex, const float& desiredTime, FVector& Vector, const bool& DoSwap, const bool& DoScale )
+bool UHoudiniCSV::GetPointVectorValueAtTime( const int32& PointID, const int32& ColumnIndex, const float& desiredTime, FVector& Vector, const bool& DoSwap, const bool& DoScale ) const
 {
 	int32 PrevIndex = -1;
 	int32 NextIndex = -1;
@@ -990,7 +972,7 @@ bool UHoudiniCSV::GetPointVectorValueAtTime( const int32& PointID, const int32&
 	return true;
 }
 
-bool UHoudiniCSV::GetRowIndexesForPointAtTime(const int32& PointID, const float& desiredTime, int32& PrevIndex, int32& NextIndex, float& PrevWeight )
+bool UHoudiniCSV::GetRowIndexesForPointAtTime(const int32& PointID, const float& desiredTime, int32& PrevIndex, int32& NextIndex, float& PrevWeight ) const
 {
 	float PrevTime = -1.0f;
 	float NextTime = -1.0f;
@@ -1019,15 +1001,15 @@ bool UHoudiniCSV::GetRowIndexesForPointAtTime(const int32& PointID, const float&
 
 	// We don't have Id information
 	// return PointID for the rowIndexes ??
-	if ( IDColumnIndex == INDEX_NONE )
+	if (!IsValidAttributeColumnIndex(EHoudiniAttributes::POINTID))
 		return false;
 
 	// We don't have time information
-	if ( TimeColumnIndex == INDEX_NONE )
+	if (!IsValidAttributeColumnIndex(EHoudiniAttributes::TIME))
 		return false;
 
 	// Get the row indexes for this point
-	TArray<int32>* RowIndexes = nullptr;
+	const TArray<int32>* RowIndexes = nullptr;
 	if ( PointValueIndexes.IsValidIndex( PointID ) )
 		RowIndexes = &( PointValueIndexes[ PointID ].RowIndexes );
 
@@ -1092,10 +1074,29 @@ bool UHoudiniCSV::GetRowIndexesForPointAtTime(const int32& PointID, const float&
 	return true;
 }
 
+int32 UHoudiniCSV::GetAttributeColumnIndex(const EHoudiniAttributes& Attr) const
+{
+	if ( !IsValidAttributeColumnIndex( Attr ) )
+		return INDEX_NONE;
 
+	return SpecialAttributesColumnIndexes[ Attr ];
+}
+
+// Returns if the specific attribute has a valid column index
+bool UHoudiniCSV::IsValidAttributeColumnIndex(const EHoudiniAttributes& Attr) const
+{
+	if ( !SpecialAttributesColumnIndexes.IsValidIndex( Attr ) )
+		return false;
+
+	int32 ColIdx = SpecialAttributesColumnIndexes[ Attr ];
+	if ( ( ColIdx < 0 ) || ( ColIdx >= NumberOfColumns) )
+		return false;
+
+	return true;
+}
 
 // Returns the column index for a given string
-bool UHoudiniCSV::GetColumnIndexFromString( const FString& ColumnTitle, int32& ColumnIndex )
+bool UHoudiniCSV::GetColumnIndexFromString( const FString& ColumnTitle, int32& ColumnIndex ) const
 {
     if ( !ColumnTitleArray.Find( ColumnTitle, ColumnIndex ) )
 		return true;
@@ -1105,12 +1106,13 @@ bool UHoudiniCSV::GetColumnIndexFromString( const FString& ColumnTitle, int32& C
 		return ColumnTitleArray.Find( TEXT( "Px" ), ColumnIndex );
     else if ( ColumnTitle.Equals( "N" ) )
 		return ColumnTitleArray.Find( TEXT( "Nx" ), ColumnIndex );
+	// TODO COLOR etc...
 
     return false;
 }
 
 // Returns the float value at a given point in the CSV file
-bool UHoudiniCSV::GetFloatValue( const int32& rowIndex, const FString& ColumnTitle, float& value )
+bool UHoudiniCSV::GetFloatValue( const int32& rowIndex, const FString& ColumnTitle, float& value ) const
 {
     int32 ColIndex = -1;
     if ( !GetColumnIndexFromString( ColumnTitle, ColIndex ) )
@@ -1179,34 +1181,34 @@ UHoudiniCSV::GetAssetRegistryTags(TArray< FAssetRegistryTag > & OutTags) const
 
 	// And a list of the special attributes we found
 	FString SpecialAttr;
-	if ( IDColumnIndex != INDEX_NONE )
+	if (IsValidAttributeColumnIndex(EHoudiniAttributes::POINTID))
 		SpecialAttr += TEXT("ID ");
 
-	if ( TypeColumnIndex != INDEX_NONE )
+	if (IsValidAttributeColumnIndex(EHoudiniAttributes::TYPE))
 		SpecialAttr += TEXT("Type ");
 
-	if ( PositionColumnIndex != INDEX_NONE )
+	if (IsValidAttributeColumnIndex(EHoudiniAttributes::POSITION))
 		SpecialAttr += TEXT("Position ");
 
-	if ( NormalColumnIndex != INDEX_NONE )
+	if (IsValidAttributeColumnIndex(EHoudiniAttributes::NORMAL))
 		SpecialAttr += TEXT("Normal ");
 
-	if ( VelocityColumnIndex != INDEX_NONE )
+	if (IsValidAttributeColumnIndex(EHoudiniAttributes::VELOCITY))
 		SpecialAttr += TEXT("Velocity ");
 
-	if ( TimeColumnIndex != INDEX_NONE )
+	if (IsValidAttributeColumnIndex(EHoudiniAttributes::TIME))
 		SpecialAttr += TEXT("Time ");
 
-	if ( ColorColumnIndex != INDEX_NONE )
+	if (IsValidAttributeColumnIndex(EHoudiniAttributes::COLOR))
 		SpecialAttr += TEXT("Color ");
 
-	if ( AlphaColumnIndex != INDEX_NONE )
+	if (IsValidAttributeColumnIndex(EHoudiniAttributes::ALPHA))
 		SpecialAttr += TEXT("Alpha ");
 
-	if ( AliveColumnIndex != INDEX_NONE )
+	if (IsValidAttributeColumnIndex(EHoudiniAttributes::ALIVE))
 		SpecialAttr += TEXT("Alive ");
 
-	if ( LifeColumnIndex != INDEX_NONE )
+	if (IsValidAttributeColumnIndex(EHoudiniAttributes::LIFE))
 		SpecialAttr += TEXT("Life ");
 
 	OutTags.Add(FAssetRegistryTag("Special Attributes", SpecialAttr, FAssetRegistryTag::TT_Alphabetical));
diff --git a/Engine/Plugins/FX/HoudiniNiagara/Source/HoudiniNiagara/Public/HoudiniCSV.h b/Engine/Plugins/FX/HoudiniNiagara/Source/HoudiniNiagara/Public/HoudiniCSV.h
--- a/Engine/Plugins/FX/HoudiniNiagara/Source/HoudiniNiagara/Public/HoudiniCSV.h
+++ b/Engine/Plugins/FX/HoudiniNiagara/Source/HoudiniNiagara/Public/HoudiniCSV.h
@@ -29,6 +29,26 @@
 
 DECLARE_LOG_CATEGORY_EXTERN( LogHoudiniNiagara, All, All );
 
+UENUM()
+enum EHoudiniAttributes
+{
+	HOUDINI_ATTR_BEGIN,
+
+	POSITION = HOUDINI_ATTR_BEGIN,
+	NORMAL,
+	TIME,
+	POINTID,
+	ALIVE,
+	LIFE,
+	COLOR,
+	ALPHA,
+	VELOCITY,
+	TYPE,
+
+	HOUDINI_ATTR_SIZE,
+	HOUDINI_ATTR_END = HOUDINI_ATTR_SIZE - 1
+	
+};
 
 USTRUCT()
 struct FPointIndexes
@@ -43,6 +63,8 @@ struct FPointIndexes
 UCLASS()
 class HOUDININIAGARA_API UHoudiniCSV : public UObject
 {
+	friend class UNiagaraDataInterfaceHoudiniCSV;
+
     GENERATED_UCLASS_BODY()
  
     public:
@@ -50,28 +72,34 @@ class HOUDININIAGARA_API UHoudiniCSV : public UObject
 	//-----------------------------------------------------------------------------------------
 	//  MEMBER FUNCTIONS
 	//-----------------------------------------------------------------------------------------
-		
+
 	bool UpdateFromFile( const FString& TheFileName );
 	bool UpdateFromStringArray( TArray<FString>& StringArray );
 
 	void SetFileName( const FString& TheFilename );
 
 	// Returns the number of points found in the CSV file
-	int32 GetNumberOfPoints();
+	int32 GetNumberOfPoints() const;
 
 	// Returns the number of rows found in the CSV file
-	int32 GetNumberOfRows();
+	int32 GetNumberOfRows() const ;
 
 	// Returns the number of columns found in the CSV file
-	int32 GetNumberOfColumns();
+	int32 GetNumberOfColumns() const;
+
+	// Return the column index for a specific attribute
+	int32 GetAttributeColumnIndex(const EHoudiniAttributes& Attr) const;
+
+	// Returns if the specific attribute has a valid column index
+	bool IsValidAttributeColumnIndex(const EHoudiniAttributes& Attr) const;
 
 	// Returns the column index for a given string
-	bool GetColumnIndexFromString(const FString& ColumnTitle, int32& ColumnIndex);
+	bool GetColumnIndexFromString(const FString& ColumnTitle, int32& ColumnIndex) const;
 
 	// Returns the float value at a given point in the CSV file
-	bool GetFloatValue( const int32& rowIndex, const int32& colIndex, float& value );
+	bool GetFloatValue( const int32& rowIndex, const int32& colIndex, float& value ) const;
 	// Returns the float value at a given point in the CSV file
-	bool GetFloatValue( const int32& rowIndex, const FString& ColumnTitle, float& value );
+	bool GetFloatValue( const int32& rowIndex, const FString& ColumnTitle, float& value ) const;
 	/*
 	// Returns the float value at a given point in the CSV file
 	bool GetCSVStringValue( const int32& rowIndex, const int32& colIndex, FString& value );
@@ -79,49 +107,50 @@ class HOUDININIAGARA_API UHoudiniCSV : public UObject
 	bool GetCSVStringValue( const int32& rowIndex, const FString& ColumnTitle, FString& value );
 	*/
 	// Returns a Vector3 for a given point in the CSV file
-	bool GetVectorValue( const int32& rowIndex, const int32& colIndex, FVector& value, const bool& DoSwap = true, const bool& DoScale = true );
+	bool GetVectorValue( const int32& rowIndex, const int32& colIndex, FVector& value, const bool& DoSwap = true, const bool& DoScale = true ) const;
 	
 	// Returns a time value for a given point in the CSV file
-	bool GetTimeValue( const int32& rowIndex, float& value );
+	bool GetTimeValue( const int32& rowIndex, float& value ) const;
 	// Returns a Position Vector3 for a given point in the CSV file (converted to unreal's coordinate system)
-	bool GetPositionValue( const int32& rowIndex, FVector& value );
+	bool GetPositionValue( const int32& rowIndex, FVector& value ) const;
 	// Returns a Normal Vector3 for a given point in the CSV file (converted to unreal's coordinate system)
-	bool GetNormalValue( const int32& rowIndex, FVector& value );
+	bool GetNormalValue( const int32& rowIndex, FVector& value ) const;
 	// Returns a Color for a given point in the CSV file
-	bool GetColorValue( const int32& rowIndex, FLinearColor& value );
+	bool GetColorValue( const int32& rowIndex, FLinearColor& value ) const;
 	// Returns a Velocity Vector3 for a given point in the CSV file
-	bool GetVelocityValue(const int32& rowIndex, FVector& value );
+	bool GetVelocityValue(const int32& rowIndex, FVector& value ) const;
 
 	// Get the last row index for a given time value (the row with a time smaller or equal to desiredTime)
 	// If the CSV file doesn't have time informations, returns false and set LastRowIndex to the last row in the file
 	// If desiredTime is smaller than the time value in the first row, LastRowIndex will be set to -1
 	// If desiredTime is higher than the last time value in the last row of the csv file, LastIndex will be set to the last row's index
-	bool GetLastRowIndexAtTime( const float& desiredTime, int32& lastRowIndex );
+	bool GetLastRowIndexAtTime( const float& desiredTime, int32& lastRowIndex ) const;
 
 	// Get the last pointID of the points to be spawned at time t
 	// Invalid Index are used to indicate edge cases:
 	// -1 will be returned if there is no points to spawn ( t is smaller than the first point time )
 	// NumberOfRows will be returned if all points in the CSV have been spawned ( t is higher than the last point time )
-	bool GetLastPointIDToSpawnAtTime( const float& time, int32& lastIndex );
+	bool GetLastPointIDToSpawnAtTime( const float& time, int32& lastIndex ) const;
 
 	// Returns the previous and next indexes for reading the values of a specified point at a given time
-	bool GetRowIndexesForPointAtTime(const int32& PointID, const float& desiredTime, int32& PrevIndex, int32& NextIndex, float& PrevWeight);
+	bool GetRowIndexesForPointAtTime(const int32& PointID, const float& desiredTime, int32& PrevIndex, int32& NextIndex, float& PrevWeight) const;
 	// Returns the value for a point at a given time value (linearly interpolated) 
-	bool GetPointValueAtTime(const int32& PointID, const int32& ColumnIndex, const float& desiredTime, float& Value);
+	bool GetPointValueAtTime(const int32& PointID, const int32& ColumnIndex, const float& desiredTime, float& Value) const;
 	// Returns the Vector Value for a given point at a given time value (linearly interpolated) 
-	bool GetPointVectorValueAtTime(const int32& PointID, const int32& ColumnIndex, const float& desiredTime, FVector& Vector, const bool& DoSwap, const bool& DoScale);
+	bool GetPointVectorValueAtTime(const int32& PointID, const int32& ColumnIndex, const float& desiredTime, FVector& Vector, const bool& DoSwap, const bool& DoScale) const;
 	// Returns the Position Value for a given point at a given time value (linearly interpolated) 
-	bool GetPointPositionAtTime(const int32& PointID, const float& desiredTime, FVector& Vector);
+	bool GetPointPositionAtTime(const int32& PointID, const float& desiredTime, FVector& Vector) const;
 	// Return a given point's life value at spawn
-	bool GetPointLife(const int32& PointID, float& Value);
+	bool GetPointLife(const int32& PointID, float& Value) const;
 	// Return a point's life for a given time value
-	bool GetPointLifeAtTime(const int32& PointID, const float& DesiredTime, float& Value);
+	bool GetPointLifeAtTime(const int32& PointID, const float& DesiredTime, float& Value) const;
 	// Return a point's type at spawn
-	bool GetPointType(const int32& PointID, int32& Value);
+	bool GetPointType(const int32& PointID, int32& Value) const;
 
 	//-----------------------------------------------------------------------------------------
 	//  MEMBER VARIABLES
 	//-----------------------------------------------------------------------------------------
+
 	UPROPERTY( VisibleAnywhere, Category = "Houdini CSV File Properties" )
 	FString FileName;
 
@@ -155,9 +184,7 @@ class HOUDININIAGARA_API UHoudiniCSV : public UObject
 	class UAssetImportData* AssetImportData;
 
 	virtual void PostInitProperties() override;
-
 	virtual void PostEditChangeProperty(FPropertyChangedEvent & PropertyChangedEvent) override;
-
 	virtual void GetAssetRegistryTags(TArray< FAssetRegistryTag > & OutTags) const override;
 #endif
 
@@ -171,9 +198,11 @@ class HOUDININIAGARA_API UHoudiniCSV : public UObject
 
 	static bool SortPredicate(const TArray<FString>& A, const TArray<FString>& B);
 
+	/*
 	// Array containing the Raw String data
-	//UPROPERTY()
-	//TArray<FString> StringCSVData;
+	UPROPERTY()
+	TArray<FString> StringCSVData;
+	*/
 
 	// Array containing all the CSV data converted to floats
 	UPROPERTY()
@@ -191,49 +220,15 @@ class HOUDININIAGARA_API UHoudiniCSV : public UObject
 	UPROPERTY()
 	TArray<int32> PointTypes;
 
-	// Index of the Position values in the buffer
-	UPROPERTY()
-	int32 PositionColumnIndex;
-
-	// Index of the Normal values in the buffer
-	UPROPERTY()
-	int32 NormalColumnIndex;
-
-	// Index of the time values in the buffer
+	// Array containing the column indexes of the special attributes
 	UPROPERTY()
-	int32 TimeColumnIndex;
-
-	// Index of the point id values in the buffer
-	UPROPERTY()
-	int32 IDColumnIndex;
-
-	// Index of the alive values in the buffer
-	UPROPERTY()
-	int32 AliveColumnIndex;
-
-	// Index of the life values in the buffer
-	UPROPERTY()
-	int32 LifeColumnIndex;
-
-	// Index of the color values in the buffer
-	UPROPERTY()
-	int32 ColorColumnIndex;
-
-	// Index of the color values in the buffer
-	UPROPERTY()
-	int32 AlphaColumnIndex;
-
-	// Index of the velocity values in the buffer
-	UPROPERTY()
-	int32 VelocityColumnIndex;
-
-	// Column index of the "type" values in the buffer
-	UPROPERTY()
-	int32 TypeColumnIndex;
+	TArray<int32> SpecialAttributesColumnIndexes;
 
+	/*
 	// Row indexes for new time values
-	//UPROPERTY()
-	//TMap<float, int32> TimeValuesIndexes;
+	UPROPERTY()
+	TMap<float, int32> TimeValuesIndexes;
+	*/
 
 	// Row indexes for each point
 	UPROPERTY()
