Houdini-Niagara:

- Added support for GPU simulations to the Houdini Data Interface.
- Houdini CSV Asset now have an icon.

diff --git a/Engine/Plugins/FX/HoudiniNiagara/Source/HoudiniNiagara/Classes/NiagaraDataInterfaceHoudiniCSV.h b/Engine/Plugins/FX/HoudiniNiagara/Source/HoudiniNiagara/Classes/NiagaraDataInterfaceHoudiniCSV.h
--- a/Engine/Plugins/FX/HoudiniNiagara/Source/HoudiniNiagara/Classes/NiagaraDataInterfaceHoudiniCSV.h
+++ b/Engine/Plugins/FX/HoudiniNiagara/Source/HoudiniNiagara/Classes/NiagaraDataInterfaceHoudiniCSV.h
@@ -37,23 +37,27 @@ UCLASS(EditInlineNew, Category = "Houdini Niagara", meta = (DisplayName = "Houdi
 class HOUDININIAGARA_API UNiagaraDataInterfaceHoudiniCSV : public UNiagaraDataInterface
 {
 	GENERATED_UCLASS_BODY()
+
 public:
 
+	// Houdini CSV Asset to sample
 	UPROPERTY( EditAnywhere, Category = "Houdini Niagara", meta = (DisplayName = "Houdini CSV Asset" ) )
 	UHoudiniCSV* HoudiniCSVAsset;
 
 	//----------------------------------------------------------------------------
-	//UObject Interface
+	// UObject Interface
 	virtual void PostInitProperties() override;
 	virtual void PostLoad() override;
 #if WITH_EDITOR
 	virtual void PostEditChangeProperty(struct FPropertyChangedEvent& PropertyChangedEvent) override;
 #endif
 	//----------------------------------------------------------------------------
+	// UNiagaraDataInterface Interface
 
+	// Returns the DI's functions signatures 
 	virtual void GetFunctions(TArray<FNiagaraFunctionSignature>& OutFunctions)override;
 	
-	/** Returns the delegate for the passed function signature. */
+	// Returns the delegate for the passed function signature.
 	virtual void GetVMExternalFunction(const FVMExternalFunctionBindingInfo& BindingInfo, void* InstanceData, FVMExternalFunction &OutFunc) override;
 
 	virtual bool Equals(const UNiagaraDataInterface* Other) const override;
@@ -155,9 +159,46 @@ public:
 	// GPU / HLSL Functions
 	virtual bool GetFunctionHLSL(const FName& DefinitionFunctionName, FString InstanceFunctionName, FNiagaraDataInterfaceGPUParamInfo& ParamInfo, FString& OutHLSL) override;
 	virtual void GetParameterDefinitionHLSL(FNiagaraDataInterfaceGPUParamInfo& ParamInfo, FString& OutHLSL) override;
-
-	// Disabling GPU sim for now.
-	virtual bool CanExecuteOnTarget(ENiagaraSimTarget Target)const override { return Target == ENiagaraSimTarget::CPUSim; }
+	virtual FNiagaraDataInterfaceParametersCS* ConstructComputeParameters() const override;
+
+	virtual bool CanExecuteOnTarget(ENiagaraSimTarget Target)const override { return true; }
+
+	// Members for GPU compatibility
+
+	// Buffer and member variables base name
+	static const FString NumberOfRowsBaseName;
+	static const FString NumberOfColumnsBaseName;
+	static const FString NumberOfPointsBaseName;
+	static const FString FloatValuesBufferBaseName;
+	static const FString SpecialAttributesColumnIndexesBufferBaseName;
+	static const FString SpawnTimesBufferBaseName;
+	static const FString LifeValuesBufferBaseName;
+	static const FString PointTypesBufferBaseName;
+	static const FString MaxNumberOfIndexesPerPointBaseName;
+	static const FString PointValueIndexesBufferBaseName;
+	static const FString LastSpawnedPointIdBaseName;
+
+	// Member variables accessors
+	FORCEINLINE int32 GetNumberOfRows()const { return HoudiniCSVAsset ? HoudiniCSVAsset->GetNumberOfRows() : 0; }
+	FORCEINLINE int32 GetNumberOfColumns()const { return HoudiniCSVAsset ? HoudiniCSVAsset->GetNumberOfColumns() : 0; }
+	FORCEINLINE int32 GetNumberOfPoints()const { return HoudiniCSVAsset ? HoudiniCSVAsset->GetNumberOfPoints() : 0; }
+	FORCEINLINE int32 GetMaxNumberOfIndexesPerPoints()const { return HoudiniCSVAsset ? HoudiniCSVAsset->GetMaxNumberOfPointValueIndexes() + 1 : 0; }
+
+	// GPU Buffers accessors
+	FRWBuffer& GetFloatValuesGPUBuffer();
+	FRWBuffer& GetSpecialAttributesColumnIndexesGPUBuffer();
+	FRWBuffer& GetSpawnTimesGPUBuffer();
+	FRWBuffer& GetLifeValuesGPUBuffer();
+	FRWBuffer& GetPointTypesGPUBuffer();
+	FRWBuffer& GetPointValueIndexesGPUBuffer();
+
+	// GPU Buffers
+	FRWBuffer FloatValuesGPUBuffer;
+	FRWBuffer SpecialAttributesColumnIndexesGPUBuffer;
+	FRWBuffer SpawnTimesGPUBuffer;
+	FRWBuffer LifeValuesGPUBuffer;
+	FRWBuffer PointTypesGPUBuffer;
+	FRWBuffer PointValueIndexesGPUBuffer;
 
 protected:
 
@@ -165,7 +206,17 @@ protected:
 
 	// Indicates the GPU buffers need to be updated
 	UPROPERTY()
-	bool GPUBufferDirty;
+	bool FloatValuesGPUBufferDirty;
+	UPROPERTY()
+	bool SpecialAttributesColumnIndexesGPUBufferDirty;
+	UPROPERTY()
+	bool SpawnTimesGPUBufferDirty;
+	UPROPERTY()
+	bool LifeValuesGPUBufferDirty;
+	UPROPERTY()
+	bool PointTypesGPUBufferDirty;
+	UPROPERTY()
+	bool PointValueIndexesGPUBufferDirty;
 
 	// Last Spawned PointID
 	UPROPERTY()
diff --git a/Engine/Plugins/FX/HoudiniNiagara/Source/HoudiniNiagara/Private/HoudiniCSV.cpp b/Engine/Plugins/FX/HoudiniNiagara/Source/HoudiniNiagara/Private/HoudiniCSV.cpp
--- a/Engine/Plugins/FX/HoudiniNiagara/Source/HoudiniNiagara/Private/HoudiniCSV.cpp
+++ b/Engine/Plugins/FX/HoudiniNiagara/Source/HoudiniNiagara/Private/HoudiniCSV.cpp
@@ -972,6 +972,18 @@ bool UHoudiniCSV::GetPointVectorValueAtTime( const int32& PointID, const int32&
 	return true;
 }
 
+int32 UHoudiniCSV::GetMaxNumberOfPointValueIndexes() const
+{
+	int32 MaxNum = 0;
+	for ( auto ValueIndexes : PointValueIndexes )
+	{
+		if ( MaxNum < ValueIndexes.RowIndexes.Num() )
+			MaxNum = ValueIndexes.RowIndexes.Num();
+	}
+
+	return MaxNum;
+}
+
 bool UHoudiniCSV::GetRowIndexesForPointAtTime(const int32& PointID, const float& desiredTime, int32& PrevIndex, int32& NextIndex, float& PrevWeight ) const
 {
 	float PrevTime = -1.0f;
diff --git a/Engine/Plugins/FX/HoudiniNiagara/Source/HoudiniNiagara/Private/NiagaraDataInterfaceHoudiniCSV.cpp b/Engine/Plugins/FX/HoudiniNiagara/Source/HoudiniNiagara/Private/NiagaraDataInterfaceHoudiniCSV.cpp
--- a/Engine/Plugins/FX/HoudiniNiagara/Source/HoudiniNiagara/Private/NiagaraDataInterfaceHoudiniCSV.cpp
+++ b/Engine/Plugins/FX/HoudiniNiagara/Source/HoudiniNiagara/Private/NiagaraDataInterfaceHoudiniCSV.cpp
@@ -24,8 +24,53 @@
 #include "NiagaraDataInterfaceHoudiniCSV.h"
 #include "NiagaraTypes.h"
 #include "Misc/FileHelper.h"
-
-#define LOCTEXT_NAMESPACE "HoudiniNiagaraCSVDataInterface"  
+#include "NiagaraShader.h"
+
+#define LOCTEXT_NAMESPACE "HoudiniNiagaraCSVDataInterface"
+
+// Base name for the member variables / buffers for GPU compatibility
+const FString UNiagaraDataInterfaceHoudiniCSV::NumberOfRowsBaseName(TEXT("NumberOfRows_"));
+const FString UNiagaraDataInterfaceHoudiniCSV::NumberOfColumnsBaseName(TEXT("NumberOfColumns_"));
+const FString UNiagaraDataInterfaceHoudiniCSV::NumberOfPointsBaseName(TEXT("NumberOfPoints_"));
+const FString UNiagaraDataInterfaceHoudiniCSV::FloatValuesBufferBaseName(TEXT("FloatValuesBuffer_"));
+const FString UNiagaraDataInterfaceHoudiniCSV::SpecialAttributesColumnIndexesBufferBaseName(TEXT("SpecialAttributesColumnIndexesBuffer_"));
+const FString UNiagaraDataInterfaceHoudiniCSV::SpawnTimesBufferBaseName(TEXT("SpawnTimesBuffer_"));
+const FString UNiagaraDataInterfaceHoudiniCSV::LifeValuesBufferBaseName(TEXT("LifeValuesBuffer_"));
+const FString UNiagaraDataInterfaceHoudiniCSV::PointTypesBufferBaseName(TEXT("PointTypesBuffer_"));
+const FString UNiagaraDataInterfaceHoudiniCSV::MaxNumberOfIndexesPerPointBaseName(TEXT("MaxNumberOfIndexesPerPoint_"));
+const FString UNiagaraDataInterfaceHoudiniCSV::PointValueIndexesBufferBaseName(TEXT("PointValueIndexesBuffer_"));
+const FString UNiagaraDataInterfaceHoudiniCSV::LastSpawnedPointIdBaseName(TEXT("LastSpawnedPointId_"));
+
+
+// Name of all the functions available in the data interface
+static const FName GetFloatValueName("GetFloatValue");
+//static const FName GetFloatValueByStringName("GetFloatValueByString");
+static const FName GetVectorValueName("GetVectorValue");
+static const FName GetVectorValueExName("GetVectorValueEx");
+
+static const FName GetPositionName("GetPosition");
+static const FName GetNormalName("GetNormal");
+static const FName GetTimeName("GetTime");
+static const FName GetVelocityName("GetVelocity");
+static const FName GetColorName("GetColor");
+static const FName GetPositionAndTimeName("GetPositionAndTime");
+
+static const FName GetNumberOfPointsName("GetNumberOfPoints");
+static const FName GetNumberOfRowsName("GetNumberOfRows");
+static const FName GetNumberOfColumnsName("GetNumberOfColumns");
+
+static const FName GetLastRowIndexAtTimeName("GetLastRowIndexAtTime");
+static const FName GetPointIDsToSpawnAtTimeName("GetPointIDsToSpawnAtTime");
+static const FName GetRowIndexesForPointAtTimeName("GetRowIndexesForPointAtTime");
+
+static const FName GetPointPositionAtTimeName("GetPointPositionAtTime");
+static const FName GetPointValueAtTimeName("GetPointValueAtTime");
+static const FName GetPointVectorValueAtTimeName("GetPointVectorValueAtTime");
+
+static const FName GetPointVectorValueAtTimeExName("GetPointVectorValueAtTimeEx");
+static const FName GetPointLifeName("GetPointLife");
+//static const FName GetPointLifeAtTimeName("GetPointLifeAtTime");
+static const FName GetPointTypeName("GetPointType");
 
 
 UNiagaraDataInterfaceHoudiniCSV::UNiagaraDataInterfaceHoudiniCSV(FObjectInitializer const& ObjectInitializer)
@@ -44,14 +89,27 @@ void UNiagaraDataInterfaceHoudiniCSV::PostInitProperties()
 	    FNiagaraTypeRegistry::Register(FNiagaraTypeDefinition(GetClass()), true, false, false);
     }
 
-    GPUBufferDirty = true;
+	FloatValuesGPUBufferDirty = true;
+	SpecialAttributesColumnIndexesGPUBufferDirty = true;
+	SpawnTimesGPUBufferDirty = true;
+	LifeValuesGPUBufferDirty = true;
+	PointTypesGPUBufferDirty = true;
+	PointValueIndexesGPUBufferDirty = true;
+
 	LastSpawnedPointID = -1;
 }
 
 void UNiagaraDataInterfaceHoudiniCSV::PostLoad()
 {
     Super::PostLoad();
-    GPUBufferDirty = true;
+
+	FloatValuesGPUBufferDirty = true;
+	SpecialAttributesColumnIndexesGPUBufferDirty = true;
+	SpawnTimesGPUBufferDirty = true;
+	LifeValuesGPUBufferDirty = true;
+	PointTypesGPUBufferDirty = true;
+	PointValueIndexesGPUBufferDirty = true;
+
 	LastSpawnedPointID = -1;
 }
 
@@ -64,9 +122,15 @@ void UNiagaraDataInterfaceHoudiniCSV::PostEditChangeProperty(struct FPropertyCha
     if (PropertyChangedEvent.Property && PropertyChangedEvent.Property->GetFName() == GET_MEMBER_NAME_CHECKED(UNiagaraDataInterfaceHoudiniCSV, HoudiniCSVAsset))
     {
 		Modify();
-		if (HoudiniCSVAsset)
+		if ( HoudiniCSVAsset )
 		{
-			GPUBufferDirty = true;
+			FloatValuesGPUBufferDirty = true;
+			SpecialAttributesColumnIndexesGPUBufferDirty = true;
+			SpawnTimesGPUBufferDirty = true;
+			LifeValuesGPUBufferDirty = true;
+			PointTypesGPUBufferDirty = true;
+			PointValueIndexesGPUBufferDirty = true;
+
 			LastSpawnedPointID = -1;
 		}
     }
@@ -104,13 +168,13 @@ bool UNiagaraDataInterfaceHoudiniCSV::Equals(const UNiagaraDataInterface* Other)
     return false;
 }
 
-
+// Returns the signature of all the functions avaialable in the data interface
 void UNiagaraDataInterfaceHoudiniCSV::GetFunctions(TArray<FNiagaraFunctionSignature>& OutFunctions)
 {
     {
 		// GetFloatValue
 		FNiagaraFunctionSignature Sig;
-		Sig.Name = TEXT("GetFloatValue");
+		Sig.Name = GetFloatValueName;
 		Sig.bMemberFunction = true;
 		Sig.bRequiresContext = false;
 		Sig.Inputs.Add(FNiagaraVariable(FNiagaraTypeDefinition(GetClass()), TEXT("CSV")));			// CSV in
@@ -128,11 +192,11 @@ void UNiagaraDataInterfaceHoudiniCSV::GetFunctions(TArray<FNiagaraFunctionSignat
     {
 		// GetFloatValueByString
 		FNiagaraFunctionSignature Sig;
-		Sig.Name = TEXT("GetFloatValueByString");
+		Sig.Name = GetFloatValueByStringName;
 		Sig.bMemberFunction = true;
 		Sig.bRequiresContext = false;
 		Sig.Inputs.Add(FNiagaraVariable(FNiagaraTypeDefinition(GetClass()), TEXT("CSV")));			// CSV in
-		Sig.Inputs.Add(FNiagaraVariable(FNiagaraTypeDefinition::GetIntDef(), TEXT("Row")));		// Row Index In
+		Sig.Inputs.Add(FNiagaraVariable(FNiagaraTypeDefinition::GetIntDef(), TEXT("Row")));			// Row Index In
 		Sig.Inputs.Add(FNiagaraVariable(FNiagaraTypeDefinition::GetStringDef(), TEXT("ColTitle")));	// Col Title In
 		Sig.Outputs.Add(FNiagaraVariable(FNiagaraTypeDefinition::GetFloatDef(), TEXT("Value")));	// Float Out
 
@@ -146,7 +210,7 @@ void UNiagaraDataInterfaceHoudiniCSV::GetFunctions(TArray<FNiagaraFunctionSignat
 	{
 		// GetVectorValue
 		FNiagaraFunctionSignature Sig;
-		Sig.Name = TEXT("GetVectorValue");
+		Sig.Name = GetVectorValueName;
 		Sig.bMemberFunction = true;
 		Sig.bRequiresContext = false;
 		Sig.Inputs.Add(FNiagaraVariable(FNiagaraTypeDefinition(GetClass()), TEXT("CSV")));			// CSV in
@@ -163,7 +227,7 @@ void UNiagaraDataInterfaceHoudiniCSV::GetFunctions(TArray<FNiagaraFunctionSignat
 	{
 		// GetVectorValueEx
 		FNiagaraFunctionSignature Sig;
-		Sig.Name = TEXT("GetVectorValueEx");
+		Sig.Name = GetVectorValueExName;
 		Sig.bMemberFunction = true;
 		Sig.bRequiresContext = false;
 		Sig.Inputs.Add(FNiagaraVariable(FNiagaraTypeDefinition(GetClass()), TEXT("CSV")));			// CSV in
@@ -182,7 +246,7 @@ void UNiagaraDataInterfaceHoudiniCSV::GetFunctions(TArray<FNiagaraFunctionSignat
     {
 		// GetPosition
 		FNiagaraFunctionSignature Sig;
-		Sig.Name = TEXT("GetPosition");
+		Sig.Name = GetPositionName;
 		Sig.bMemberFunction = true;
 		Sig.bRequiresContext = false;
 		Sig.Inputs.Add(FNiagaraVariable(FNiagaraTypeDefinition(GetClass()), TEXT("CSV")));			// CSV in
@@ -198,7 +262,7 @@ void UNiagaraDataInterfaceHoudiniCSV::GetFunctions(TArray<FNiagaraFunctionSignat
     {
 		// GetPositionAndTime
 		FNiagaraFunctionSignature Sig;
-		Sig.Name = TEXT("GetPositionAndTime");
+		Sig.Name = GetPositionAndTimeName;
 		Sig.bMemberFunction = true;
 		Sig.bRequiresContext = false;
 		Sig.Inputs.Add(FNiagaraVariable(FNiagaraTypeDefinition(GetClass()), TEXT("CSV")));			// CSV in
@@ -215,7 +279,7 @@ void UNiagaraDataInterfaceHoudiniCSV::GetFunctions(TArray<FNiagaraFunctionSignat
     {
 		// GetNormal
 		FNiagaraFunctionSignature Sig;
-		Sig.Name = TEXT("GetNormal");
+		Sig.Name = GetNormalName;
 		Sig.bMemberFunction = true;
 		Sig.bRequiresContext = false;
 		Sig.Inputs.Add(FNiagaraVariable(FNiagaraTypeDefinition(GetClass()), TEXT("CSV")));			// CSV in
@@ -231,7 +295,7 @@ void UNiagaraDataInterfaceHoudiniCSV::GetFunctions(TArray<FNiagaraFunctionSignat
     {
 		// GetTime
 		FNiagaraFunctionSignature Sig;
-		Sig.Name = TEXT("GetTime");
+		Sig.Name = GetTimeName;
 		Sig.bMemberFunction = true;
 		Sig.bRequiresContext = false;
 		Sig.Inputs.Add(FNiagaraVariable(FNiagaraTypeDefinition(GetClass()), TEXT("CSV")));			// CSV in
@@ -247,7 +311,7 @@ void UNiagaraDataInterfaceHoudiniCSV::GetFunctions(TArray<FNiagaraFunctionSignat
 	{
 		// GetVelocity
 		FNiagaraFunctionSignature Sig;
-		Sig.Name = TEXT("GetVelocity");
+		Sig.Name = GetVelocityName;
 		Sig.bMemberFunction = true;
 		Sig.bRequiresContext = false;
 		Sig.Inputs.Add(FNiagaraVariable(FNiagaraTypeDefinition(GetClass()), TEXT("CSV")));			// CSV in
@@ -263,7 +327,7 @@ void UNiagaraDataInterfaceHoudiniCSV::GetFunctions(TArray<FNiagaraFunctionSignat
 	{
 		// GetColor
 		FNiagaraFunctionSignature Sig;
-		Sig.Name = TEXT("GetColor");
+		Sig.Name = GetColorName;
 		Sig.bMemberFunction = true;
 		Sig.bRequiresContext = false;
 		Sig.Inputs.Add(FNiagaraVariable(FNiagaraTypeDefinition(GetClass()), TEXT("CSV")));			// CSV in
@@ -279,7 +343,7 @@ void UNiagaraDataInterfaceHoudiniCSV::GetFunctions(TArray<FNiagaraFunctionSignat
     {
 		// GetNumberOfPoints
 		FNiagaraFunctionSignature Sig;
-		Sig.Name = TEXT("GetNumberOfPoints");
+		Sig.Name = GetNumberOfPointsName;
 		Sig.bMemberFunction = true;
 		Sig.bRequiresContext = false;
 		Sig.Inputs.Add(FNiagaraVariable(FNiagaraTypeDefinition(GetClass()), TEXT("CSV")));					// CSV in
@@ -294,7 +358,7 @@ void UNiagaraDataInterfaceHoudiniCSV::GetFunctions(TArray<FNiagaraFunctionSignat
 	{
 		// GetNumberOfRows
 		FNiagaraFunctionSignature Sig;
-		Sig.Name = TEXT("GetNumberOfRows");
+		Sig.Name = GetNumberOfRowsName;
 		Sig.bMemberFunction = true;
 		Sig.bRequiresContext = false;
 		Sig.Inputs.Add( FNiagaraVariable( FNiagaraTypeDefinition( GetClass()), TEXT("CSV") ) );					// CSV in
@@ -309,7 +373,7 @@ void UNiagaraDataInterfaceHoudiniCSV::GetFunctions(TArray<FNiagaraFunctionSignat
 	{
 		// GetNumberOfColumns
 		FNiagaraFunctionSignature Sig;
-		Sig.Name = TEXT("GetNumberOfColumns");
+		Sig.Name = GetNumberOfColumnsName;
 		Sig.bMemberFunction = true;
 		Sig.bRequiresContext = false;
 		Sig.Inputs.Add( FNiagaraVariable( FNiagaraTypeDefinition( GetClass()), TEXT("CSV") ) );						// CSV in
@@ -324,7 +388,7 @@ void UNiagaraDataInterfaceHoudiniCSV::GetFunctions(TArray<FNiagaraFunctionSignat
     {
 		// GetLastRowIndexAtTime
 		FNiagaraFunctionSignature Sig;
-		Sig.Name = TEXT("GetLastRowIndexAtTime");
+		Sig.Name = GetLastRowIndexAtTimeName;
 		Sig.bMemberFunction = true;
 		Sig.bRequiresContext = false;
 		Sig.Inputs.Add(FNiagaraVariable( FNiagaraTypeDefinition(GetClass()), TEXT("CSV") ) );					// CSV in
@@ -340,7 +404,7 @@ void UNiagaraDataInterfaceHoudiniCSV::GetFunctions(TArray<FNiagaraFunctionSignat
     {
 		// GetPointIdsToSpawnAtTime
 		FNiagaraFunctionSignature Sig;
-		Sig.Name = TEXT("GetPointIDsToSpawnAtTime");
+		Sig.Name = GetPointIDsToSpawnAtTimeName;
 		Sig.bMemberFunction = true;
 		Sig.bRequiresContext = false;
 		Sig.Inputs.Add(FNiagaraVariable( FNiagaraTypeDefinition(GetClass()), TEXT("CSV") ) );				// CSV in
@@ -358,7 +422,7 @@ void UNiagaraDataInterfaceHoudiniCSV::GetFunctions(TArray<FNiagaraFunctionSignat
 	{
 		// GetRowIndexesForPointAtTime
 		FNiagaraFunctionSignature Sig;
-		Sig.Name = TEXT("GetRowIndexesForPointAtTime");
+		Sig.Name = GetRowIndexesForPointAtTimeName;
 		Sig.bMemberFunction = true;
 		Sig.bRequiresContext = false;
 		Sig.Inputs.Add(FNiagaraVariable(FNiagaraTypeDefinition(GetClass()), TEXT("CSV")));					// CSV in
@@ -377,11 +441,11 @@ void UNiagaraDataInterfaceHoudiniCSV::GetFunctions(TArray<FNiagaraFunctionSignat
 	{
 		// GetPointPositionAtTime
 		FNiagaraFunctionSignature Sig;
-		Sig.Name = TEXT("GetPointPositionAtTime");
+		Sig.Name = GetPointPositionAtTimeName;
 		Sig.bMemberFunction = true;
 		Sig.bRequiresContext = false;
 		Sig.Inputs.Add(FNiagaraVariable(FNiagaraTypeDefinition(GetClass()), TEXT("CSV")));				// CSV in
-		Sig.Inputs.Add(FNiagaraVariable(FNiagaraTypeDefinition::GetIntDef(), TEXT("PointID")));				// Point Number In
+		Sig.Inputs.Add(FNiagaraVariable(FNiagaraTypeDefinition::GetIntDef(), TEXT("PointID")));			// Point Number In
 		Sig.Inputs.Add(FNiagaraVariable(FNiagaraTypeDefinition::GetFloatDef(), TEXT("Time")));		    // Time in
 		Sig.Outputs.Add(FNiagaraVariable(FNiagaraTypeDefinition::GetVec3Def(), TEXT("Position")));		// Vector3 Out
 
@@ -394,11 +458,11 @@ void UNiagaraDataInterfaceHoudiniCSV::GetFunctions(TArray<FNiagaraFunctionSignat
 	{
 		// GetPointValueAtTime
 		FNiagaraFunctionSignature Sig;
-		Sig.Name = TEXT("GetPointValueAtTime");
+		Sig.Name = GetPointValueAtTimeName;
 		Sig.bMemberFunction = true;
 		Sig.bRequiresContext = false;
 		Sig.Inputs.Add(FNiagaraVariable(FNiagaraTypeDefinition(GetClass()), TEXT("CSV")));				// CSV in
-		Sig.Inputs.Add(FNiagaraVariable(FNiagaraTypeDefinition::GetIntDef(), TEXT("PointID")));				// Point Number In		
+		Sig.Inputs.Add(FNiagaraVariable(FNiagaraTypeDefinition::GetIntDef(), TEXT("PointID")));			// Point Number In		
 		Sig.Inputs.Add(FNiagaraVariable(FNiagaraTypeDefinition::GetIntDef(), TEXT("Col")));				// Col Index In
 		Sig.Inputs.Add(FNiagaraVariable(FNiagaraTypeDefinition::GetFloatDef(), TEXT("Time")));		    // Time in		
 		Sig.Outputs.Add(FNiagaraVariable(FNiagaraTypeDefinition::GetFloatDef(), TEXT("Value")));		// Float Out
@@ -412,7 +476,7 @@ void UNiagaraDataInterfaceHoudiniCSV::GetFunctions(TArray<FNiagaraFunctionSignat
 	{
 		// GetPointVectorValueAtTime
 		FNiagaraFunctionSignature Sig;
-		Sig.Name = TEXT("GetPointVectorValueAtTime");
+		Sig.Name = GetPointVectorValueAtTimeName;
 		Sig.bMemberFunction = true;
 		Sig.bRequiresContext = false;
 		Sig.Inputs.Add(FNiagaraVariable(FNiagaraTypeDefinition(GetClass()), TEXT("CSV")));				// CSV in
@@ -430,7 +494,7 @@ void UNiagaraDataInterfaceHoudiniCSV::GetFunctions(TArray<FNiagaraFunctionSignat
 	{
 		// GetPointVectorValueAtTimeEx
 		FNiagaraFunctionSignature Sig;
-		Sig.Name = TEXT("GetPointVectorValueAtTimeEx");
+		Sig.Name = GetPointVectorValueAtTimeExName;
 		Sig.bMemberFunction = true;
 		Sig.bRequiresContext = false;
 		Sig.Inputs.Add(FNiagaraVariable(FNiagaraTypeDefinition(GetClass()), TEXT("CSV")));				// CSV in
@@ -450,7 +514,7 @@ void UNiagaraDataInterfaceHoudiniCSV::GetFunctions(TArray<FNiagaraFunctionSignat
 	{
 		// GetPointLife
 		FNiagaraFunctionSignature Sig;
-		Sig.Name = TEXT("GetPointLife");
+		Sig.Name = GetPointLifeName;
 		Sig.bMemberFunction = true;
 		Sig.bRequiresContext = false;
 		Sig.Inputs.Add(FNiagaraVariable(FNiagaraTypeDefinition(GetClass()), TEXT("CSV")));				// CSV in
@@ -467,7 +531,7 @@ void UNiagaraDataInterfaceHoudiniCSV::GetFunctions(TArray<FNiagaraFunctionSignat
 	{
 		// GetPointLifeAtTime
 		FNiagaraFunctionSignature Sig;
-		Sig.Name = TEXT("GetPointLifeAtTime");
+		Sig.Name = GetPointLifeAtTimeName;
 		Sig.bMemberFunction = true;
 		Sig.bRequiresContext = false;
 		Sig.Inputs.Add(FNiagaraVariable(FNiagaraTypeDefinition(GetClass()), TEXT("CSV")));				// CSV in
@@ -485,7 +549,7 @@ void UNiagaraDataInterfaceHoudiniCSV::GetFunctions(TArray<FNiagaraFunctionSignat
 	{
 		// GetPointType
 		FNiagaraFunctionSignature Sig;
-		Sig.Name = TEXT("GetPointType");
+		Sig.Name = GetPointTypeName;
 		Sig.bMemberFunction = true;
 		Sig.bRequiresContext = false;
 		Sig.Inputs.Add(FNiagaraVariable(FNiagaraTypeDefinition(GetClass()), TEXT("CSV")));				// CSV in
@@ -499,89 +563,6 @@ void UNiagaraDataInterfaceHoudiniCSV::GetFunctions(TArray<FNiagaraFunctionSignat
 	}
 }
 
-// build the shader function HLSL; function name is passed in, as it's defined per-DI; that way, configuration could change
-// the HLSL in the spirit of a static switch
-// TODO: need a way to identify each specific function here
-// 
-bool UNiagaraDataInterfaceHoudiniCSV::GetFunctionHLSL(const FName& DefinitionFunctionName, FString InstanceFunctionName, FNiagaraDataInterfaceGPUParamInfo& ParamInfo, FString& OutHLSL)
-{
-    //FString BufferName = Descriptors[0].BufferParamName;
-    //FString SecondBufferName = Descriptors[1].BufferParamName;
-
-    /*
-    if (InstanceFunctionName.Contains( TEXT( "GetFloatValue") ) )
-    {
-	FString BufferName = Descriptors[0].BufferParamName;
-	OutHLSL += TEXT("void ") + InstanceFunctionName + TEXT("(in float In_Row, in float In_Col, out float Out_Value) \n{\n");
-	OutHLSL += TEXT("\t Out_Value = ") + BufferName + TEXT("[(int)(In_Row + ( In_Col * ") + FString::FromInt(NumberOfRows) + TEXT(") ) ];");
-	OutHLSL += TEXT("\n}\n");
-    }
-    else if (InstanceFunctionName.Contains(TEXT("GetPosition")))
-    {
-	FString BufferName = Descriptors[1].BufferParamName;
-	OutHLSL += TEXT("void ") + InstanceFunctionName + TEXT("(in float In_N, out float3 Out_Value) \n{\n");
-	OutHLSL += TEXT("\t Out_Value.x = ") + BufferName + TEXT("[(int)(In_N) ];");
-	OutHLSL += TEXT("\t Out_Value.y = ") + BufferName + TEXT("[(int)(In_N + ") + FString::FromInt(NumberOfRows) + TEXT(") ];");
-	OutHLSL += TEXT("\t Out_Value.z = ") + BufferName + TEXT("[(int)(In_N + ( 2 * ") + FString::FromInt(NumberOfRows) + TEXT(") ) ];");
-	OutHLSL += TEXT("\n}\n");
-    }
-    else if (InstanceFunctionName.Contains(TEXT("GetNormal")))
-    {
-	FString BufferName = Descriptors[2].BufferParamName;
-	OutHLSL += TEXT("void ") + InstanceFunctionName + TEXT("(in float In_N, out float3 Out_Value) \n{\n");
-	OutHLSL += TEXT("\t Out_Value.x = ") + BufferName + TEXT("[(int)(In_N) ];");
-	OutHLSL += TEXT("\t Out_Value.y = ") + BufferName + TEXT("[(int)(In_N + ") + FString::FromInt(NumberOfRows) + TEXT(") ];");
-	OutHLSL += TEXT("\t Out_Value.z = ") + BufferName + TEXT("[(int)(In_N + ( 2 * ") + FString::FromInt(NumberOfRows) + TEXT(") ) ];");
-	OutHLSL += TEXT("\n}\n");
-    }
-    else if (InstanceFunctionName.Contains(TEXT("GetTime")))
-    {
-	FString BufferName = Descriptors[3].BufferParamName;
-	OutHLSL += TEXT("void ") + InstanceFunctionName + TEXT("(in float In_N, out float Out_Value) \n{\n");
-	OutHLSL += TEXT("\t Out_Value = ") + BufferName + TEXT("[(int)(In_N) ];");
-	OutHLSL += TEXT("\n}\n");
-    }
-    else if (InstanceFunctionName.Contains( TEXT("GetNumberOfPointsInCSV") ) )
-    {
-	OutHLSL += TEXT("void ") + InstanceFunctionName + TEXT("( out int Out_Value ) \n{\n");
-	OutHLSL += TEXT("\t Out_Value = ") + FString::FromInt(NumberOfRows) + TEXT(";");
-	OutHLSL += TEXT("\n}\n");
-    }*/
-    /*else if (InstanceFunctionName.Contains(TEXT("GetPositionAndTime")))
-    {
-	OutHLSL += TEXT("void ") + FunctionName + TEXT("(in float In_N, out float4 Out_Value) \n{\n");
-	OutHLSL += TEXT("\t Out_Value.x = ") + BufferName + TEXT("[(int)(In_N) ];");
-	OutHLSL += TEXT("\t Out_Value.y = ") + BufferName + TEXT("[(int)(In_N + ") + FString::FromInt(NumberOfRows) + TEXT(") ];");
-	OutHLSL += TEXT("\t Out_Value.z = ") + BufferName + TEXT("[(int)(In_N + ( 2 * ") + FString::FromInt(NumberOfRows) + TEXT(") ) ];");
-	OutHLSL += TEXT("\t Out_Value.w = ") + SecondBufferName + TEXT("[(int)(In_N) ];");
-	OutHLSL += TEXT("\n}\n");
-    }*/
-
-    return !OutHLSL.IsEmpty();
-}
-
-// build buffer definition hlsl
-// 1. Choose a buffer name, add the data interface ID (important!)
-// 2. add a DIGPUBufferParamDescriptor to the array argument; that'll be passed on to the FNiagaraShader for binding to a shader param, that can
-// then later be found by name via FindDIBufferParam for setting; 
-// 3. store buffer declaration hlsl in OutHLSL
-// multiple buffers can be defined at once here
-//
-void UNiagaraDataInterfaceHoudiniCSV::GetParameterDefinitionHLSL(FNiagaraDataInterfaceGPUParamInfo& ParamInfo, FString& OutHLSL)
-{
-    FString BufferName = "CSVData_" + ParamInfo.DataInterfaceHLSLSymbol;
-    OutHLSL += TEXT("Buffer<float> ") + BufferName + TEXT(";\n");
-
-    BufferName = "PositionData_" + ParamInfo.DataInterfaceHLSLSymbol;
-    OutHLSL += TEXT("Buffer<float> ") + BufferName + TEXT(";\n");
-
-    BufferName = "NormalData_" + ParamInfo.DataInterfaceHLSLSymbol;
-    OutHLSL += TEXT("Buffer<float> ") + BufferName + TEXT(";\n");
-
-    BufferName = "TimeData_" + ParamInfo.DataInterfaceHLSLSymbol;
-    OutHLSL += TEXT("Buffer<float> ") + BufferName + TEXT(";\n");
-}
-
 DEFINE_NDI_RAW_FUNC_BINDER(UNiagaraDataInterfaceHoudiniCSV, GetFloatValue);
 //DEFINE_NDI_RAW_FUNC_BINDER(UNiagaraDataInterfaceHoudiniCSV, GetFloatValueByString);
 DEFINE_NDI_RAW_FUNC_BINDER(UNiagaraDataInterfaceHoudiniCSV, GetVectorValue);
@@ -604,97 +585,97 @@ DEFINE_NDI_RAW_FUNC_BINDER(UNiagaraDataInterfaceHoudiniCSV, GetPointLife);
 DEFINE_NDI_RAW_FUNC_BINDER(UNiagaraDataInterfaceHoudiniCSV, GetPointType);
 void UNiagaraDataInterfaceHoudiniCSV::GetVMExternalFunction(const FVMExternalFunctionBindingInfo& BindingInfo, void* InstanceData, FVMExternalFunction &OutFunc)
 {
-    if (BindingInfo.Name == TEXT("GetFloatValue") && BindingInfo.GetNumInputs() == 2 && BindingInfo.GetNumOutputs() == 1)
+    if (BindingInfo.Name == GetFloatValueName && BindingInfo.GetNumInputs() == 2 && BindingInfo.GetNumOutputs() == 1)
     {
 		TNDIParamBinder<0, int32, TNDIParamBinder<1, int32, NDI_RAW_FUNC_BINDER(UNiagaraDataInterfaceHoudiniCSV, GetFloatValue)>>::Bind(this, BindingInfo, InstanceData, OutFunc);
     }
-    /*else if (BindingInfo.Name == TEXT("GetFloatValueByString") && BindingInfo.GetNumInputs() == 2 && BindingInfo.GetNumOutputs() == 1)
+    /*else if (BindingInfo.Name == GetFloatValueByStringName && BindingInfo.GetNumInputs() == 2 && BindingInfo.GetNumOutputs() == 1)
     {
 		TNDIParamBinder<0, float, TNDIParamBinder<1, FString, NDI_RAW_FUNC_BINDER(UNiagaraDataInterfaceHoudiniCSV, GetFloatValueByString)>>::Bind(this, BindingInfo, InstanceData, OutFunc);
     }*/
-    else if (BindingInfo.Name == TEXT("GetVectorValue") && BindingInfo.GetNumInputs() == 2 && BindingInfo.GetNumOutputs() == 3)
+    else if (BindingInfo.Name == GetVectorValueName && BindingInfo.GetNumInputs() == 2 && BindingInfo.GetNumOutputs() == 3)
     {
 		TNDIParamBinder<0, int32, TNDIParamBinder<1, int32, NDI_RAW_FUNC_BINDER(UNiagaraDataInterfaceHoudiniCSV, GetVectorValue)>>::Bind(this, BindingInfo, InstanceData, OutFunc);
     }
-	else if (BindingInfo.Name == TEXT("GetVectorValueEx") && BindingInfo.GetNumInputs() == 4 && BindingInfo.GetNumOutputs() == 3)
+	else if (BindingInfo.Name == GetVectorValueExName && BindingInfo.GetNumInputs() == 4 && BindingInfo.GetNumOutputs() == 3)
 	{
 		TNDIParamBinder<0, int32, TNDIParamBinder<1, int32, TNDIParamBinder<2, bool, TNDIParamBinder<3, bool, NDI_RAW_FUNC_BINDER(UNiagaraDataInterfaceHoudiniCSV, GetVectorValueEx)>>>>::Bind(this, BindingInfo, InstanceData, OutFunc);
 	}
-    else if (BindingInfo.Name == TEXT("GetPosition") && BindingInfo.GetNumInputs() == 1 && BindingInfo.GetNumOutputs() == 3)
+    else if (BindingInfo.Name == GetPositionName && BindingInfo.GetNumInputs() == 1 && BindingInfo.GetNumOutputs() == 3)
     {
 		TNDIParamBinder<0, int32, NDI_RAW_FUNC_BINDER(UNiagaraDataInterfaceHoudiniCSV, GetPosition)>::Bind(this, BindingInfo, InstanceData, OutFunc);
     }
-    else if (BindingInfo.Name == TEXT("GetNormal") && BindingInfo.GetNumInputs() == 1 && BindingInfo.GetNumOutputs() == 3)
+    else if (BindingInfo.Name == GetNormalName && BindingInfo.GetNumInputs() == 1 && BindingInfo.GetNumOutputs() == 3)
     {
 		TNDIParamBinder<0, int32, NDI_RAW_FUNC_BINDER(UNiagaraDataInterfaceHoudiniCSV, GetNormal)>::Bind(this, BindingInfo, InstanceData, OutFunc);
     }
-    else if (BindingInfo.Name == TEXT("GetTime") && BindingInfo.GetNumInputs() == 1 && BindingInfo.GetNumOutputs() == 1)
+    else if (BindingInfo.Name == GetTimeName && BindingInfo.GetNumInputs() == 1 && BindingInfo.GetNumOutputs() == 1)
     {
 		TNDIParamBinder<0, int32, NDI_RAW_FUNC_BINDER(UNiagaraDataInterfaceHoudiniCSV, GetTime)>::Bind(this, BindingInfo, InstanceData, OutFunc);
     }
-	else if (BindingInfo.Name == TEXT("GetVelocity") && BindingInfo.GetNumInputs() == 1 && BindingInfo.GetNumOutputs() == 3)
+	else if (BindingInfo.Name == GetVelocityName && BindingInfo.GetNumInputs() == 1 && BindingInfo.GetNumOutputs() == 3)
 	{
 		TNDIParamBinder<0, int32, NDI_RAW_FUNC_BINDER(UNiagaraDataInterfaceHoudiniCSV, GetVelocity)>::Bind(this, BindingInfo, InstanceData, OutFunc);
 	}
-	else if (BindingInfo.Name == TEXT("GetColor") && BindingInfo.GetNumInputs() == 1 && BindingInfo.GetNumOutputs() == 4)
+	else if (BindingInfo.Name == GetColorName && BindingInfo.GetNumInputs() == 1 && BindingInfo.GetNumOutputs() == 4)
 	{
 		TNDIParamBinder<0, int32, NDI_RAW_FUNC_BINDER(UNiagaraDataInterfaceHoudiniCSV, GetColor)>::Bind(this, BindingInfo, InstanceData, OutFunc);
 	}
-    else if (BindingInfo.Name == TEXT("GetPositionAndTime") && BindingInfo.GetNumInputs() == 1 && BindingInfo.GetNumOutputs() == 4)
+    else if (BindingInfo.Name == GetPositionAndTimeName && BindingInfo.GetNumInputs() == 1 && BindingInfo.GetNumOutputs() == 4)
     {
 		TNDIParamBinder<0, int32, NDI_RAW_FUNC_BINDER(UNiagaraDataInterfaceHoudiniCSV, GetPositionAndTime)>::Bind(this, BindingInfo, InstanceData, OutFunc);
     }
-    else if ( BindingInfo.Name == TEXT("GetNumberOfPoints") && BindingInfo.GetNumInputs() == 0 && BindingInfo.GetNumOutputs() == 1 )
+    else if ( BindingInfo.Name == GetNumberOfPointsName && BindingInfo.GetNumInputs() == 0 && BindingInfo.GetNumOutputs() == 1 )
     {
 		OutFunc = FVMExternalFunction::CreateUObject(this, &UNiagaraDataInterfaceHoudiniCSV::GetNumberOfPoints);
     }
-	else if (BindingInfo.Name == TEXT("GetNumberOfRows") && BindingInfo.GetNumInputs() == 0 && BindingInfo.GetNumOutputs() == 1)
+	else if (BindingInfo.Name == GetNumberOfRowsName && BindingInfo.GetNumInputs() == 0 && BindingInfo.GetNumOutputs() == 1)
 	{
 		OutFunc = FVMExternalFunction::CreateUObject(this, &UNiagaraDataInterfaceHoudiniCSV::GetNumberOfRows);
 	}
-	else if (BindingInfo.Name == TEXT("GetNumberOfColumns") && BindingInfo.GetNumInputs() == 0 && BindingInfo.GetNumOutputs() == 1)
+	else if (BindingInfo.Name == GetNumberOfColumnsName && BindingInfo.GetNumInputs() == 0 && BindingInfo.GetNumOutputs() == 1)
 	{
 		OutFunc = FVMExternalFunction::CreateUObject(this, &UNiagaraDataInterfaceHoudiniCSV::GetNumberOfColumns);
 	}
-    else if (BindingInfo.Name == TEXT("GetLastRowIndexAtTime") && BindingInfo.GetNumInputs() == 1 && BindingInfo.GetNumOutputs() == 1)
+    else if (BindingInfo.Name == GetLastRowIndexAtTimeName && BindingInfo.GetNumInputs() == 1 && BindingInfo.GetNumOutputs() == 1)
     {
 		TNDIParamBinder<0, float, NDI_RAW_FUNC_BINDER(UNiagaraDataInterfaceHoudiniCSV, GetLastRowIndexAtTime)>::Bind(this, BindingInfo, InstanceData, OutFunc);
     }
-    else if (BindingInfo.Name == TEXT("GetPointIDsToSpawnAtTime") && BindingInfo.GetNumInputs() == 1 && BindingInfo.GetNumOutputs() == 3)
+    else if (BindingInfo.Name == GetPointIDsToSpawnAtTimeName && BindingInfo.GetNumInputs() == 1 && BindingInfo.GetNumOutputs() == 3)
     {
 		TNDIParamBinder<0, float, NDI_RAW_FUNC_BINDER(UNiagaraDataInterfaceHoudiniCSV, GetPointIDsToSpawnAtTime)>::Bind(this, BindingInfo, InstanceData, OutFunc);
     }
-	else if (BindingInfo.Name == TEXT("GetRowIndexesForPointAtTime") && BindingInfo.GetNumInputs() == 2 && BindingInfo.GetNumOutputs() == 3)
+	else if (BindingInfo.Name == GetRowIndexesForPointAtTimeName && BindingInfo.GetNumInputs() == 2 && BindingInfo.GetNumOutputs() == 3)
 	{
 		TNDIParamBinder<0, int32, TNDIParamBinder<1, float, NDI_RAW_FUNC_BINDER(UNiagaraDataInterfaceHoudiniCSV, GetRowIndexesForPointAtTime)>>::Bind(this, BindingInfo, InstanceData, OutFunc);
 	}
-	else if (BindingInfo.Name == TEXT("GetPointPositionAtTime") && BindingInfo.GetNumInputs() == 2 && BindingInfo.GetNumOutputs() == 3)
+	else if (BindingInfo.Name == GetPointPositionAtTimeName && BindingInfo.GetNumInputs() == 2 && BindingInfo.GetNumOutputs() == 3)
 	{
 		TNDIParamBinder<0, int32, TNDIParamBinder<1, float, NDI_RAW_FUNC_BINDER(UNiagaraDataInterfaceHoudiniCSV, GetPointPositionAtTime)>>::Bind(this, BindingInfo, InstanceData, OutFunc);
 	}
-	else if (BindingInfo.Name == TEXT("GetPointValueAtTime") && BindingInfo.GetNumInputs() == 3 && BindingInfo.GetNumOutputs() == 1)
+	else if (BindingInfo.Name == GetPointValueAtTimeName && BindingInfo.GetNumInputs() == 3 && BindingInfo.GetNumOutputs() == 1)
 	{
 		TNDIParamBinder<0, int32, TNDIParamBinder<1, int32, TNDIParamBinder<2, float, NDI_RAW_FUNC_BINDER(UNiagaraDataInterfaceHoudiniCSV, GetPointValueAtTime)>>>::Bind(this, BindingInfo, InstanceData, OutFunc);
 	}
-	else if (BindingInfo.Name == TEXT("GetPointVectorValueAtTime") && BindingInfo.GetNumInputs() == 3 && BindingInfo.GetNumOutputs() == 3)
+	else if (BindingInfo.Name == GetPointVectorValueAtTimeName && BindingInfo.GetNumInputs() == 3 && BindingInfo.GetNumOutputs() == 3)
 	{
 		TNDIParamBinder<0, int32, TNDIParamBinder<1, int32, TNDIParamBinder<2, float, NDI_RAW_FUNC_BINDER(UNiagaraDataInterfaceHoudiniCSV, GetPointVectorValueAtTime)>>>::Bind(this, BindingInfo, InstanceData, OutFunc);
 	}
-	else if (BindingInfo.Name == TEXT("GetPointVectorValueAtTimeEx") && BindingInfo.GetNumInputs() == 5 && BindingInfo.GetNumOutputs() == 3)
+	else if (BindingInfo.Name == GetPointVectorValueAtTimeExName && BindingInfo.GetNumInputs() == 5 && BindingInfo.GetNumOutputs() == 3)
 	{
 		TNDIParamBinder<0, int32, TNDIParamBinder<1, int32, TNDIParamBinder<2, float, TNDIParamBinder<3, bool, TNDIParamBinder<4, bool, NDI_RAW_FUNC_BINDER(UNiagaraDataInterfaceHoudiniCSV, GetPointVectorValueAtTimeEx)>>>>>::Bind(this, BindingInfo, InstanceData, OutFunc);
 	}
-	else if (BindingInfo.Name == TEXT("GetPointLife") && BindingInfo.GetNumInputs() == 1 && BindingInfo.GetNumOutputs() == 1)
+	else if (BindingInfo.Name == GetPointLifeName && BindingInfo.GetNumInputs() == 1 && BindingInfo.GetNumOutputs() == 1)
 	{
 		TNDIParamBinder<0, int32, NDI_RAW_FUNC_BINDER(UNiagaraDataInterfaceHoudiniCSV, GetPointLife)>::Bind(this, BindingInfo, InstanceData, OutFunc);
 	}
 	/*
-	else if (BindingInfo.Name == TEXT("GetPointLifeAtTime") && BindingInfo.GetNumInputs() == 2 && BindingInfo.GetNumOutputs() == 1)
+	else if (BindingInfo.Name == GetPointLifeAtTimeName && BindingInfo.GetNumInputs() == 2 && BindingInfo.GetNumOutputs() == 1)
 	{
 		TNDIParamBinder<0, int32, TNDIParamBinder<1, float, NDI_RAW_FUNC_BINDER(UNiagaraDataInterfaceHoudiniCSV, GetPointLifeAtTime)>>::Bind(this, BindingInfo, InstanceData, OutFunc);
 	}
 	*/
-	else if (BindingInfo.Name == TEXT("GetPointType") && BindingInfo.GetNumInputs() == 1 && BindingInfo.GetNumOutputs() == 1)
+	else if (BindingInfo.Name == GetPointTypeName && BindingInfo.GetNumInputs() == 1 && BindingInfo.GetNumOutputs() == 1)
 	{
 		TNDIParamBinder<0, int32, NDI_RAW_FUNC_BINDER(UNiagaraDataInterfaceHoudiniCSV, GetPointType)>::Bind(this, BindingInfo, InstanceData, OutFunc);
 	}
@@ -1344,8 +1325,8 @@ void UNiagaraDataInterfaceHoudiniCSV::GetPointType(FVectorVMContext& Context)
 
 void UNiagaraDataInterfaceHoudiniCSV::GetNumberOfRows(FVectorVMContext& Context)
 {
-    FRegisterHandler<int32> OutNumRows(Context);
-    *OutNumRows.GetDest() = HoudiniCSVAsset ? HoudiniCSVAsset->GetNumberOfRows() : 0;
+	FRegisterHandler<int32> OutNumRows(Context);
+	*OutNumRows.GetDest() = HoudiniCSVAsset ? HoudiniCSVAsset->GetNumberOfRows() : 0;
 	OutNumRows.Advance();
 }
 
@@ -1363,4 +1344,760 @@ void UNiagaraDataInterfaceHoudiniCSV::GetNumberOfPoints(FVectorVMContext& Contex
 	OutNumPoints.Advance();
 }
 
+// Build the shader function HLSL Code.
+bool UNiagaraDataInterfaceHoudiniCSV::GetFunctionHLSL(const FName& DefinitionFunctionName, FString InstanceFunctionName, FNiagaraDataInterfaceGPUParamInfo& ParamInfo, FString& OutHLSL)
+{
+	// Build the buffer/variable names declared for this DI
+	FString NumberOfRowsVar = NumberOfRowsBaseName + ParamInfo.DataInterfaceHLSLSymbol;
+	FString NumberOfColumnsVar = NumberOfColumnsBaseName + ParamInfo.DataInterfaceHLSLSymbol;
+	FString NumberOfPointsVar = NumberOfPointsBaseName + ParamInfo.DataInterfaceHLSLSymbol;
+	FString FloatBufferVar = FloatValuesBufferBaseName + ParamInfo.DataInterfaceHLSLSymbol;
+	FString AttributeColumnIndexesBuffer = SpecialAttributesColumnIndexesBufferBaseName + ParamInfo.DataInterfaceHLSLSymbol;
+	FString SpawnTimeBuffer = SpawnTimesBufferBaseName + ParamInfo.DataInterfaceHLSLSymbol;
+	FString LifeValuesBuffer = LifeValuesBufferBaseName + ParamInfo.DataInterfaceHLSLSymbol;
+	FString PointTypesBuffer = PointTypesBufferBaseName + ParamInfo.DataInterfaceHLSLSymbol;
+	FString MaxNumberOfIndexesPerPointVar = MaxNumberOfIndexesPerPointBaseName + ParamInfo.DataInterfaceHLSLSymbol;
+	FString PointValueIndexesBuffer = PointValueIndexesBufferBaseName + ParamInfo.DataInterfaceHLSLSymbol;
+	FString LastSpawnedPointIDVar = LastSpawnedPointIdBaseName + ParamInfo.DataInterfaceHLSLSymbol;
+
+	// Lambda returning the HLSL code used for reading a Float value in the FloatBuffer
+	auto ReadFloatInBuffer = [&](const FString& OutFloatValue, const FString& FloatRowIndex, const FString& FloatColIndex)
+	{
+		// \t OutValue = FloatBufferName[ (RowIndex) + ( (ColIndex) * (NumberOfRowsName) ) ];\n
+		return TEXT("\t ") + OutFloatValue + TEXT(" = ") + FloatBufferVar + TEXT("[ (") + FloatRowIndex + TEXT(") + ( (") + FloatColIndex + TEXT(") * (") + NumberOfRowsVar + TEXT(") ) ];\n");
+	};
+
+	// Lambda returning the HLSL code for reading a Vector value in the FloatBuffer
+	// It expects the In_DoSwap and In_DoScale integers to be defined before being called!
+	auto ReadVectorInBuffer = [&](const FString& OutVectorValue, const FString& VectorRowIndex, const FString& VectorColIndex)
+	{
+		FString OutHLSLCode;
+		OutHLSLCode += TEXT("\t// ReadVectorInBuffer\n");
+		OutHLSLCode += TEXT("\t{\n");
+			OutHLSLCode += TEXT("\t\tfloat3 temp_Value = float3(0.0, 0.0, 0.0);\n");
+			OutHLSLCode += ReadFloatInBuffer(TEXT("temp_Value.x"), VectorRowIndex, VectorColIndex);
+			OutHLSLCode += ReadFloatInBuffer(TEXT("temp_Value.y"), VectorRowIndex, VectorColIndex + TEXT(" + 1"));
+			OutHLSLCode += ReadFloatInBuffer(TEXT("temp_Value.z"), VectorRowIndex, VectorColIndex + TEXT(" + 2"));
+			OutHLSLCode += TEXT("\t\t") + OutVectorValue + TEXT(" = temp_Value;\n");
+			OutHLSLCode += TEXT("\t\tif ( In_DoSwap == 1 )\n");
+			OutHLSLCode += TEXT("\t\t{\n");
+				OutHLSLCode += TEXT("\t\t\t") + OutVectorValue + TEXT(".y= temp_Value.z;\n");
+				OutHLSLCode += TEXT("\t\t\t") + OutVectorValue + TEXT(".z= temp_Value.y;\n");
+			OutHLSLCode += TEXT("\t\t}\n");
+			OutHLSLCode += TEXT("\t\tif ( In_DoScale == 1 )\n");
+			OutHLSLCode += TEXT("\t\t{\n");
+				OutHLSLCode += TEXT("\t\t\t") + OutVectorValue + TEXT(" *= 100.0f;\n");
+			OutHLSLCode += TEXT("\t\t}\n");
+		OutHLSLCode += TEXT("\t}\n");
+		return OutHLSLCode;
+	};
+
+	// Lambda returning the HLSL code for reading a special attribute's column index in the buffer
+	auto GetSpecAttributeColumnIndex = [&](const EHoudiniAttributes& Attr)
+	{
+		FString OutHLSLCode;
+		OutHLSLCode += AttributeColumnIndexesBuffer + TEXT("[") + FString::FromInt(Attr) + TEXT("]");
+		return OutHLSLCode;
+	};
+
+	// Lambda returning the HLSL code for getting the previous and next row indexes for a given particle at a given time	
+	auto GetRowIndexesForPointAtTime = [&](const FString& In_PointID, const FString& In_Time, const FString& Out_PreviousRow, const FString& Out_NextRow, const FString& Out_Weight )
+	{
+		FString OutHLSLCode;
+		OutHLSLCode += TEXT("\t// GetRowIndexesForPointAtTime\n");
+		OutHLSLCode += TEXT("\t{\n");
+			// If the point hasn't spawn before DesiredTime
+			OutHLSLCode += TEXT("\t\tfloat spawn_time = ") + SpawnTimeBuffer + TEXT("[ (") + In_PointID + TEXT(") ];\n");
+			OutHLSLCode += TEXT("\t\tif ( spawn_time > (") + In_Time + TEXT(") ){ return; }\n");
+			// If the point is dead before DesiredTime
+			OutHLSLCode += TEXT("\t\tfloat life_value = ") + LifeValuesBuffer + TEXT("[ (") + In_PointID + TEXT(") ];\n");
+			OutHLSLCode += TEXT("\t\tif ( life_value > 0.0f && ( spawn_time + life_value < (") + In_Time + TEXT(") ) ) { return; }\n");
+
+			// Make sure we have time and id attributes
+			OutHLSLCode += TEXT("\t\tint id_col = ") + GetSpecAttributeColumnIndex(EHoudiniAttributes::POINTID) + TEXT(";\n");
+			OutHLSLCode += TEXT("\t\tif ( id_col < 0 || id_col >= ") + NumberOfColumnsVar + TEXT(" ){ return; }\n");
+			OutHLSLCode += TEXT("\t\tint time_col = ") + GetSpecAttributeColumnIndex(EHoudiniAttributes::TIME) + TEXT(";\n");
+			OutHLSLCode += TEXT("\t\tif ( time_col < 0 || time_col >= ") + NumberOfColumnsVar + TEXT(" ){ return; }\n");
+
+			OutHLSLCode += TEXT("\t\tfloat prev_time = -1.0f;\n");
+			OutHLSLCode += TEXT("\t\tfloat next_time = -1.0f;\n");
+
+			// Look at all the values for this Point
+			OutHLSLCode += TEXT("\t\tfor( int n = 0; n < ") + MaxNumberOfIndexesPerPointVar + TEXT("; n++ )\n\t{\n");
+				OutHLSLCode += TEXT("\t\t\tint current_row = ") + PointValueIndexesBuffer + TEXT("[ (") + In_PointID + TEXT(") * ") + MaxNumberOfIndexesPerPointVar + TEXT(" + n ];\n");
+				OutHLSLCode += TEXT("\t\t\tif ( current_row < 0 ){ break; }\n");
+
+				OutHLSLCode += TEXT("\t\t\tfloat current_time = -1.0f;\n");
+				OutHLSLCode += ReadFloatInBuffer(TEXT("current_time"), TEXT("current_row"), TEXT("time_col"));
+
+				OutHLSLCode += TEXT("\t\t\tif ( current_time == (") + In_Time + TEXT(") )\n");
+					OutHLSLCode += TEXT("\t\t\t\t{ ") + Out_PreviousRow + TEXT(" = current_row; ") + Out_NextRow + TEXT(" = current_row; ") + Out_Weight + TEXT(" = 1.0; return;}\n");
+				OutHLSLCode += TEXT("\t\t\telse if ( current_time < (") + In_Time + TEXT(") ){\n");
+					OutHLSLCode += TEXT("\t\t\t\tif ( prev_time == -1.0f || prev_time < current_time )\n");
+						OutHLSLCode += TEXT("\t\t\t\t\t { ") + Out_PreviousRow + TEXT(" = current_row; prev_time = current_time; }\n");
+				OutHLSLCode += TEXT("\t\t\t}\n");
+				OutHLSLCode += TEXT("\t\t\telse if ( next_time == -1.0f || next_time > current_time )\n");
+					OutHLSLCode += TEXT("\t\t\t\t { ") + Out_NextRow + TEXT(" = current_row; next_time = current_time; break; }\n");
+			OutHLSLCode += TEXT("\t\t}\n");
+
+			OutHLSLCode += TEXT("\t\tif ( ") + Out_PreviousRow + TEXT(" < 0 )\n");
+				OutHLSLCode += TEXT("\t\t\t{ ") + Out_Weight + TEXT(" = 0.0f; ") + Out_PreviousRow + TEXT(" = ") + Out_NextRow + TEXT("; return;}\n");
+			OutHLSLCode += TEXT("\t\tif ( ") + Out_NextRow + TEXT(" < 0 )\n");
+				OutHLSLCode += TEXT("\t\t\t{ ") + Out_Weight + TEXT(" = 1.0f; ") + Out_NextRow + TEXT(" = ") + Out_PreviousRow + TEXT("; return;}\n");
+
+			// Calculate the weight
+			OutHLSLCode += TEXT("\t\t") + Out_Weight + TEXT(" = ( ( (") + In_Time + TEXT(") - prev_time ) / ( next_time - prev_time ) );\n");
+
+		OutHLSLCode += TEXT("\t}\n");
+		return OutHLSLCode;
+	};
+
+	// Build each function's HLSL code
+	if (DefinitionFunctionName == GetFloatValueName)
+	{
+		// GetFloatValue(int In_Row, int In_Col, out float Out_Value)		
+		OutHLSL += TEXT("void ") + InstanceFunctionName + TEXT("(int In_Row, int In_Col, out float Out_Value) \n{\n");
+
+			OutHLSL += ReadFloatInBuffer(TEXT("Out_Value"), TEXT("In_Row"), TEXT("In_Col"));
+
+		OutHLSL += TEXT("\n}\n");
+		return true;
+	}
+	/*
+	else if (DefinitionFunctionName == GetFloatValueByStringName)
+	{
+		// GetFloatValueByString(int In_Row, string In_ColTitle, out float Out_Value)
+		FString BufferName = FloatValuesBufferBaseName + ParamInfo.DataInterfaceHLSLSymbol;
+		OutHLSL += TEXT("void ") + InstanceFunctionName + TEXT("(int In_Row, string In_ColTitle, out float Out_Value) \n{\n");
+		// NEED CODE
+		OutHLSL += TEXT("\n}\n");
+		return true;
+	}
+	*/
+	else if (DefinitionFunctionName == GetVectorValueName)
+	{
+		// GetVectorValue(int In_Row, int In_Col, out float3 Out_Value)
+		OutHLSL += TEXT("void ") + InstanceFunctionName + TEXT("(int In_Row, int In_Col, out float3 Out_Value) \n{\n");
+
+			OutHLSL += TEXT("\tint In_DoSwap = 1;\n");
+			OutHLSL += TEXT("\tint In_DoScale = 1;\n");
+			OutHLSL += ReadVectorInBuffer(TEXT("Out_Value"), TEXT("In_Row"), TEXT("In_Col"));
+
+		OutHLSL += TEXT("\n}\n");
+		return true;
+	}
+	else if (DefinitionFunctionName == GetVectorValueExName)
+	{
+		// GetVectorValueEx(int In_Row, int In_Col, int In_DoSwap, int In_DoScale, out float3 Out_Value)
+		OutHLSL += TEXT("void ") + InstanceFunctionName + TEXT("(int In_Row, int In_Col, int In_DoSwap, int In_DoScale, out float3 Out_Value) \n{\n");
+		
+			OutHLSL += ReadVectorInBuffer(TEXT("Out_Value"), TEXT("In_Row"), TEXT("In_Col"));
+
+		OutHLSL += TEXT("\n}\n");
+		return true;
+	}
+	else if (DefinitionFunctionName == GetPositionName)
+	{
+		// GetPosition(int In_Row, out float3 Out_Position)
+		OutHLSL += TEXT("void ") + InstanceFunctionName + TEXT("(int In_Row, out float3 Out_Position) \n{\n");
+
+			OutHLSL += TEXT("\tint In_DoSwap = 1;\n");
+			OutHLSL += TEXT("\tint In_DoScale = 1;\n");
+			OutHLSL += TEXT("\tint In_Col = ") + GetSpecAttributeColumnIndex( EHoudiniAttributes::POSITION ) + TEXT(";\n");
+			OutHLSL += ReadVectorInBuffer(TEXT("Out_Position"), TEXT("In_Row"), TEXT("In_Col"));
+
+		OutHLSL += TEXT("\n}\n");
+		return true;
+	}
+	else if (DefinitionFunctionName == GetPositionAndTimeName)
+	{
+		// GetPositionAndTime(int In_Row, out float3 Out_Position, out float Out_Time)
+		OutHLSL += TEXT("void ") + InstanceFunctionName + TEXT("(int In_Row, out float3 Out_Position, out float Out_Time) \n{\n");
+
+			OutHLSL += TEXT("\tint In_DoSwap = 1;\n");
+			OutHLSL += TEXT("\tint In_DoScale = 1;\n");
+			OutHLSL += TEXT("\tint In_PosCol = ") + GetSpecAttributeColumnIndex(EHoudiniAttributes::POSITION) + TEXT(";\n");
+			OutHLSL += ReadVectorInBuffer(TEXT("Out_Position"), TEXT("In_Row"), TEXT("In_PosCol"));
+
+			OutHLSL += TEXT("\tint In_TimeCol = ") + GetSpecAttributeColumnIndex(EHoudiniAttributes::TIME) + TEXT(";\n");
+			OutHLSL += ReadFloatInBuffer(TEXT("Out_Time"), TEXT("In_Row"), TEXT("In_TimeCol"));
+
+		OutHLSL += TEXT("\n}\n");
+		return true;
+	}
+	else if (DefinitionFunctionName == GetNormalName)
+	{
+		// GetNormal(int In_Row, out float3 Out_Value)
+		OutHLSL += TEXT("void ") + InstanceFunctionName + TEXT("(int In_Row, out float3 Out_Normal) \n{\n");
+
+			OutHLSL += TEXT("\tint In_DoSwap = 1;\n");
+			OutHLSL += TEXT("\tint In_DoScale = 0;\n");
+			OutHLSL += TEXT("\tint In_Col = ") + GetSpecAttributeColumnIndex(EHoudiniAttributes::NORMAL) + TEXT(";\n");
+			OutHLSL += ReadVectorInBuffer(TEXT("Out_Normal"), TEXT("In_Row"), TEXT("In_Col"));
+
+		OutHLSL += TEXT("\n}\n");
+		return true;
+	}
+	else if (DefinitionFunctionName == GetTimeName)
+	{
+		// GetTime(int In_Row, out float Out_Value)
+		OutHLSL += TEXT("void ") + InstanceFunctionName + TEXT("(int In_Row, out float Out_Time) \n{\n");
+
+			OutHLSL += TEXT("\tint In_TimeCol = ") + GetSpecAttributeColumnIndex(EHoudiniAttributes::TIME) + TEXT(";\n");
+			OutHLSL += ReadFloatInBuffer(TEXT("Out_Time"), TEXT("In_Row"), TEXT("In_TimeCol"));
+
+		OutHLSL += TEXT("\n}\n");
+		return true;
+	}
+	else if (DefinitionFunctionName == GetVelocityName)
+	{
+		// GetVelocity(int In_Row, out float3 Out_Value)
+		OutHLSL += TEXT("void ") + InstanceFunctionName + TEXT("(int In_Row, out float3 Out_Velocity) \n{\n");
+
+			OutHLSL += TEXT("\tint In_DoSwap = 1;\n");
+			OutHLSL += TEXT("\tint In_DoScale = 0;\n");
+			OutHLSL += TEXT("\tint In_Col = ") + GetSpecAttributeColumnIndex(EHoudiniAttributes::VELOCITY) + TEXT(";\n");
+			OutHLSL += ReadVectorInBuffer(TEXT("Out_Velocity"), TEXT("In_Row"), TEXT("In_Col"));
+
+		OutHLSL += TEXT("\n}\n");
+		return true;
+	}
+	else if (DefinitionFunctionName == GetColorName)
+	{
+		// GetColor(int In_Row, out float4 Out_Value)
+		OutHLSL += TEXT("void ") + InstanceFunctionName + TEXT("(int In_Row, out float4 Out_Color) \n{\n");
+
+			OutHLSL += TEXT("\tfloat3 temp_color = float3(1.0, 1.0, 1.0);\n");
+			OutHLSL += TEXT("\tint In_ColorCol = ") + GetSpecAttributeColumnIndex(EHoudiniAttributes::COLOR) + TEXT(";\n");
+			OutHLSL += TEXT("\tint In_DoSwap = 0;\n");
+			OutHLSL += TEXT("\tint In_DoScale = 0;\n");
+			OutHLSL += ReadVectorInBuffer(TEXT("temp_color"), TEXT("In_Row"), TEXT("In_ColorCol"));
+
+			OutHLSL += TEXT("\tfloat temp_alpha = 1.0;\n");
+			OutHLSL += TEXT("\tint In_TimeCol = ") + GetSpecAttributeColumnIndex(EHoudiniAttributes::TIME) + TEXT(";\n");
+			OutHLSL += ReadFloatInBuffer(TEXT("temp_alpha"), TEXT("In_Row"), TEXT("In_TimeCol"));
+
+			OutHLSL += TEXT("Out_Color.x = temp_color.x;\n");
+			OutHLSL += TEXT("Out_Color.y = temp_color.y;\n");
+			OutHLSL += TEXT("Out_Color.z = temp_color.z;\n");
+			OutHLSL += TEXT("Out_Color.w = temp_alpha;\n");
+
+		OutHLSL += TEXT("\n}\n");
+		return true;
+	}
+	else if (DefinitionFunctionName == GetNumberOfPointsName)
+	{
+		// GetNumberOfPoints(out int Out_Value)
+		OutHLSL += TEXT("void ") + InstanceFunctionName + TEXT("(out int Out_NumPoints) \n{\n");
+			OutHLSL += TEXT("\tOut_NumPoints = ") + NumberOfPointsVar + TEXT(";\n");
+		OutHLSL += TEXT("\n}\n");
+		return true;
+	}
+	else if (DefinitionFunctionName == GetNumberOfRowsName)
+	{
+		// GetNumberOfRows(out int Out_Value)
+		OutHLSL += TEXT("void ") + InstanceFunctionName + TEXT("(out int Out_NumRows) \n{\n");
+			OutHLSL += TEXT("\tOut_NumRows = ") + NumberOfRowsVar + TEXT(";\n");
+		OutHLSL += TEXT("\n}\n");
+		return true;
+	}
+	else if (DefinitionFunctionName == GetNumberOfColumnsName)
+	{
+		// GetNumberOfColumns(out int Out_Value)
+		OutHLSL += TEXT("void ") + InstanceFunctionName + TEXT("(out int Out_NumColumns) \n{\n");
+			OutHLSL += TEXT("\tOut_NumColumns = ") + NumberOfColumnsVar + TEXT(";\n");
+		OutHLSL += TEXT("\n}\n");
+		return true;
+	}
+	else if (DefinitionFunctionName == GetLastRowIndexAtTimeName)
+	{
+		// GetLastRowIndexAtTime(float In_Time, out int Out_Value)
+		OutHLSL += TEXT("void ") + InstanceFunctionName + TEXT("(float In_Time, out int Out_Value) \n{\n");
+
+			OutHLSL += TEXT("\tint In_TimeCol = ") + GetSpecAttributeColumnIndex( EHoudiniAttributes::TIME ) + TEXT(";\n");
+			OutHLSL += TEXT("\tfloat temp_time = 1.0;\n");
+			OutHLSL += ReadFloatInBuffer(TEXT("temp_time"), NumberOfRowsVar + TEXT("- 1"), TEXT("In_TimeCol"));
+			OutHLSL += TEXT("\tif ( temp_time < In_Time ) { Out_Value = ") + NumberOfRowsVar + TEXT(" - 1; return; }\n");
+			OutHLSL += TEXT("\tint lastRow = -1;\n");
+			OutHLSL += TEXT("\tfor( int n = 0; n < ") + NumberOfRowsVar + TEXT("; n++ )\n\t{\n");
+				OutHLSL += TEXT("\t") + ReadFloatInBuffer(TEXT("temp_time"), TEXT("n"), TEXT("In_TimeCol"));
+				OutHLSL += TEXT("\t\tif ( temp_time == In_Time ){ lastRow = n ;}");
+				OutHLSL += TEXT("\t\telse if ( temp_time > In_Time ){ lastRow = n -1; return;}");
+				OutHLSL += TEXT("\t\tif ( lastRow == -1 ){ lastRow = ") + NumberOfRowsVar + TEXT(" - 1; }");
+			OutHLSL += TEXT("\t}\n");
+
+		OutHLSL += TEXT("\n}\n");
+		return true;
+	}
+	else if (DefinitionFunctionName == GetPointIDsToSpawnAtTimeName)
+	{
+		// GetPointIDsToSpawnAtTime(float In_Time, out int Out_MinIndex, out int Out_MaxIndex, out int Out_Count)
+		OutHLSL += TEXT("void ") + InstanceFunctionName + TEXT("(float In_Time, out int Out_MinIndex, out int Out_MaxIndex, out int Out_Count) \n{\n");
+			
+			OutHLSL += TEXT("\tint time_col = ") + GetSpecAttributeColumnIndex(EHoudiniAttributes::TIME) + TEXT(";\n");
+			OutHLSL += TEXT("\tif( time_col < 0 )\n");
+				OutHLSL += TEXT("\t\t{ Out_Count = ") + NumberOfPointsVar +TEXT("; Out_MinIndex = 0; Out_MaxIndex = Out_Count - 1; return; }\n");
+
+			// GetLastPointIDToSpawnAtTime
+			OutHLSL += TEXT("\tint last_id = -1;\n");
+			OutHLSL += TEXT("\tfloat temp_time = 0;\n");
+			OutHLSL += TEXT("\tfor( int n = 0; n < ") + NumberOfPointsVar + TEXT("; n++ )\n\t{\n");
+				OutHLSL += TEXT("\t\ttemp_time = ") + SpawnTimeBuffer + TEXT("[ n ];\n");
+				OutHLSL += TEXT("\t\tif ( temp_time == In_Time )\n");
+					OutHLSL += TEXT("\t\t\t{ last_id = n; }\n");
+				OutHLSL += TEXT("\t\telse if ( temp_time > In_Time )\n");
+					OutHLSL += TEXT("\t\t\t{ last_id = n - 1; break; }\n");
+			OutHLSL += TEXT("\t}\n");
+
+			// Check if we need to reset lastspawnedPointID
+			OutHLSL += TEXT("\tif ( last_id < ") + LastSpawnedPointIDVar + TEXT(")\n");
+				OutHLSL += TEXT("\t\t{") + LastSpawnedPointIDVar + TEXT(" = -1; }\n");
+			// Nothing to spawn, In_Time is lower than the point's time
+			OutHLSL += TEXT("\tif ( last_id < 0 )\n");
+				OutHLSL += TEXT("\t\t{") + LastSpawnedPointIDVar + TEXT(" = -1;Out_Count = 0; Out_MinIndex = 0; Out_MaxIndex = 0;return;}\n");
+			// We dont have any new point to spawn
+			OutHLSL += TEXT("\tif ( last_id == ") + LastSpawnedPointIDVar + TEXT(")\n");
+				OutHLSL += TEXT("{ Out_MinIndex = last_id; Out_MaxIndex = last_id; Out_Count = 0; return;}\n");
+
+			OutHLSL += TEXT("\tOut_MinIndex = ") + LastSpawnedPointIDVar + TEXT(" + 1;\n");
+			OutHLSL += TEXT("\tOut_MaxIndex = last_id;\n");
+			OutHLSL += TEXT("\tOut_Count = Out_MaxIndex - Out_MinIndex + 1;\n");
+			OutHLSL += TEXT("\t") + LastSpawnedPointIDVar + TEXT(" = Out_MaxIndex;\n");
+
+		OutHLSL += TEXT("\n}\n");
+		return true;
+	}
+	else if (DefinitionFunctionName == GetRowIndexesForPointAtTimeName)
+	{
+		// GetRowIndexesForPointAtTime(int In_PointID, float In_Time, out int Out_PreviousRow, out int Out_NextRow, out float Out_Weight)
+		OutHLSL += TEXT("void ") + InstanceFunctionName + TEXT("(int In_PointID, float In_Time, out int Out_PreviousRow, out int Out_NextRow, out float Out_Weight) \n{\n");
+		OutHLSL += TEXT("\tOut_PreviousRow = 0;\n\tOut_NextRow = 0;\n\tOut_Weight = 0;\n");
+
+			OutHLSL += GetRowIndexesForPointAtTime(TEXT("In_PointID"), TEXT("In_Time"), TEXT("Out_PreviousRow"), TEXT("Out_NextRow"), TEXT("Out_Weight"));
+
+		OutHLSL += TEXT("\n}\n");
+		return true;
+	}
+	else if (DefinitionFunctionName == GetPointPositionAtTimeName)
+	{
+		// GetPointPositionAtTime(int In_PointID, float In_Time, out float3 Out_Value)
+		OutHLSL += TEXT("void ") + InstanceFunctionName + TEXT("(int In_PointID, float In_Time, out float3 Out_Value) \n{\n");
+		OutHLSL += TEXT("Out_Value = float3( 0.0, 0.0, 0.0 );\n");
+			OutHLSL += TEXT("\tint pos_col = ") + GetSpecAttributeColumnIndex(EHoudiniAttributes::POSITION) + TEXT(";\n");
+			OutHLSL += TEXT("\tint prev_index = -1;int next_index = -1;float weight = 1.0f;\n");
+			OutHLSL += GetRowIndexesForPointAtTime(TEXT("In_PointID"), TEXT("In_Time"), TEXT("prev_index"), TEXT("next_index"), TEXT("weight"));
+
+			OutHLSL += TEXT("\tint In_DoSwap = 1;\n");
+			OutHLSL += TEXT("\tint In_DoScale = 1;\n");
+
+			OutHLSL += TEXT("\tfloat3 prev_vector = float3( 0.0, 0.0, 0.0 );\n");
+			OutHLSL += ReadVectorInBuffer(TEXT("prev_vector"), TEXT("prev_index"), TEXT("pos_col"));
+			OutHLSL += TEXT("\tfloat3 next_vector = float3( 0.0, 0.0, 0.0 );\n");
+			OutHLSL += ReadVectorInBuffer(TEXT("next_vector"), TEXT("next_index"), TEXT("pos_col"));
+
+			OutHLSL += TEXT("\tOut_Value = lerp(prev_vector, next_vector, weight);\n");
+
+		OutHLSL += TEXT("\n}\n");
+		return true;
+	}
+	else if (DefinitionFunctionName == GetPointValueAtTimeName)
+	{
+		// GetPointValueAtTime(int In_PointID, int In_Col, float In_Time, out float Out_Value)
+		OutHLSL += TEXT("void ") + InstanceFunctionName + TEXT("(int In_PointID, int In_Col, float In_Time, out float Out_Value) \n{\n");
+		
+		OutHLSL += TEXT("\tint prev_index = -1;int next_index = -1;float weight = 1.0f;\n");
+		OutHLSL += GetRowIndexesForPointAtTime(TEXT("In_PointID"), TEXT("In_Time"), TEXT("prev_index"), TEXT("next_index"), TEXT("weight"));
+
+		OutHLSL += TEXT("\tfloat prev_value;\n");
+		OutHLSL += ReadFloatInBuffer(TEXT("prev_value"), TEXT("prev_index"), TEXT("In_Col"));
+		OutHLSL += TEXT("\tfloat next_value;\n");
+		OutHLSL += ReadFloatInBuffer(TEXT("next_value"), TEXT("next_index"), TEXT("In_Col"));
+
+		OutHLSL += TEXT("\tOut_Value = lerp(prev_value, next_value, weight);\n");
+
+		OutHLSL += TEXT("\n}\n");
+		return true;
+	}
+	else if (DefinitionFunctionName == GetPointVectorValueAtTimeName)
+	{
+		// GetPointVectorValueAtTime(int In_PointID, int In_Col, float In_Time, out float3 Out_Value)
+		OutHLSL += TEXT("void ") + InstanceFunctionName + TEXT("(int In_PointID, int In_Col, float In_Time, out float3 Out_Value) \n{\n");
+
+		OutHLSL += TEXT("\tint prev_index = -1;int next_index = -1;float weight = 1.0f;\n");
+		OutHLSL += GetRowIndexesForPointAtTime(TEXT("In_PointID"), TEXT("In_Time"), TEXT("prev_index"), TEXT("next_index"), TEXT("weight"));
+
+		OutHLSL += TEXT("\tint In_DoSwap = 1;\n");
+		OutHLSL += TEXT("\tint In_DoScale = 1;\n");
+
+		OutHLSL += TEXT("\tfloat3 prev_vector;\n");
+		OutHLSL += ReadVectorInBuffer(TEXT("prev_vector"), TEXT("prev_index"), TEXT("In_Col"));
+		OutHLSL += TEXT("\tfloat3 next_vector;\n");
+		OutHLSL += ReadVectorInBuffer(TEXT("next_vector"), TEXT("next_index"), TEXT("In_Col"));
+
+		OutHLSL += TEXT("\tOut_Value = lerp(prev_vector, next_vector, weight);\n");
+
+		OutHLSL += TEXT("\n}\n");
+		return true;
+	}
+	else if (DefinitionFunctionName == GetPointVectorValueAtTimeExName)
+	{
+		// GetPointVectorValueAtTimeEx(int In_PointID, int In_Col, float In_Time, int In_DoSwap, int In_DoScale, out float3 Out_Value)
+		OutHLSL += TEXT("void ") + InstanceFunctionName + TEXT("(int In_PointID, int In_Col, float In_Time, int In_DoSwap, int In_DoScale, out float3 Out_Value) \n{\n");
+
+			OutHLSL += TEXT("\tint prev_index = -1;int next_index = -1;float weight = 1.0f;\n");
+			OutHLSL += GetRowIndexesForPointAtTime(TEXT("In_PointID"), TEXT("In_Time"), TEXT("prev_index"), TEXT("next_index"), TEXT("weight"));
+
+			OutHLSL += TEXT("\tfloat3 prev_vector;\n");
+			OutHLSL += ReadVectorInBuffer(TEXT("prev_vector"), TEXT("prev_index"), TEXT("In_Col"));
+			OutHLSL += TEXT("\tfloat3 next_vector;\n");
+			OutHLSL += ReadVectorInBuffer(TEXT("next_vector"), TEXT("next_index"), TEXT("In_Col"));
+
+			OutHLSL += TEXT("\tOut_Value = lerp(prev_vector, next_vector, weight);\n");
+
+		OutHLSL += TEXT("\n}\n");
+
+		return true;
+	}
+	else if (DefinitionFunctionName == GetPointLifeName)
+	{
+		// GetPointLife(int In_PointID, out float Out_Value)
+		OutHLSL += TEXT("void ") + InstanceFunctionName + TEXT("(int In_PointID, out float Out_Value) \n{\n");
+
+			OutHLSL += TEXT("\tif ( ( In_PointID < 0 ) || ( In_PointID >= ") + NumberOfPointsVar + TEXT(") )\n");
+				OutHLSL += TEXT("\t\t{Out_Value = -1.0f; return;}\n");
+			OutHLSL += TEXT("\tOut_Value = ") + LifeValuesBuffer + TEXT("[ In_PointID ];\n");
+
+		OutHLSL += TEXT("\n}\n");
+		return true;
+	}
+	/*
+	else if (DefinitionFunctionName == GetPointLifeAtTimeName)
+	{
+		// GetPointLifeAtTime(int In_PointID, float In_Time, out float Out_Value)
+		OutHLSL += TEXT("void ") + InstanceFunctionName + TEXT("(int In_PointID, float In_Time, out float Out_Value) \n{\n");
+
+		OutHLSL += TEXT("\n}\n");
+		return true;
+	}
+	*/
+	else if (DefinitionFunctionName == GetPointTypeName)
+	{
+		// GetPointType(int In_PointID, out int Out_Value)
+		OutHLSL += TEXT("void ") + InstanceFunctionName + TEXT("(int In_PointID, out int Out_Value) \n{\n");
+
+			OutHLSL += TEXT("\tif ( ( In_PointID < 0 ) || ( In_PointID >= ") + NumberOfPointsVar+ TEXT(") )\n");
+				OutHLSL += TEXT("\t\t{Out_Value = -1; return;}\n");
+			OutHLSL += TEXT("\tOut_Value = ") + PointTypesBuffer + TEXT("[ In_PointID ];\n");
+
+		OutHLSL += TEXT("\n}\n");
+		return true;
+	}
+
+	return !OutHLSL.IsEmpty();
+}
+
+// Build buffers and member variables HLSL definition
+// Always use the BaseName + the DataInterfaceHLSL indentifier!
+void UNiagaraDataInterfaceHoudiniCSV::GetParameterDefinitionHLSL(FNiagaraDataInterfaceGPUParamInfo& ParamInfo, FString& OutHLSL)
+{
+	// int NumberOfRows_XX;
+	FString BufferName = UNiagaraDataInterfaceHoudiniCSV::NumberOfRowsBaseName + ParamInfo.DataInterfaceHLSLSymbol;
+	OutHLSL += TEXT("int ") + BufferName + TEXT(";\n");
+
+	// int NumberOfColumns_XX;
+	BufferName = UNiagaraDataInterfaceHoudiniCSV::NumberOfColumnsBaseName + ParamInfo.DataInterfaceHLSLSymbol;
+	OutHLSL += TEXT("int ") + BufferName + TEXT(";\n");
+
+	// int NumberOfPoints_XX;
+	BufferName = UNiagaraDataInterfaceHoudiniCSV::NumberOfPointsBaseName + ParamInfo.DataInterfaceHLSLSymbol;
+	OutHLSL += TEXT("int ") + BufferName + TEXT(";\n");
+
+	// Buffer<float> FloatValuesBuffer_XX;
+	BufferName = UNiagaraDataInterfaceHoudiniCSV::FloatValuesBufferBaseName + ParamInfo.DataInterfaceHLSLSymbol;
+	OutHLSL += TEXT("Buffer<float> ") + BufferName + TEXT(";\n");
+
+	// Buffer<int> SpecialAttributesColumnIndexesBuffer_XX;
+	BufferName = UNiagaraDataInterfaceHoudiniCSV::SpecialAttributesColumnIndexesBufferBaseName + ParamInfo.DataInterfaceHLSLSymbol;
+	OutHLSL += TEXT("Buffer<int> ") + BufferName + TEXT(";\n");
+
+	// Buffer<float> SpawnTimesBuffer_XX;
+	BufferName = UNiagaraDataInterfaceHoudiniCSV::SpawnTimesBufferBaseName + ParamInfo.DataInterfaceHLSLSymbol;
+	OutHLSL += TEXT("Buffer<float> ") + BufferName + TEXT(";\n");
+
+	// Buffer<float> LifeValuesBuffer_XX;
+	BufferName = UNiagaraDataInterfaceHoudiniCSV::LifeValuesBufferBaseName + ParamInfo.DataInterfaceHLSLSymbol;
+	OutHLSL += TEXT("Buffer<float> ") + BufferName + TEXT(";\n");
+
+	// Buffer<int> PointTypesBuffer_XX;
+	BufferName = UNiagaraDataInterfaceHoudiniCSV::PointTypesBufferBaseName + ParamInfo.DataInterfaceHLSLSymbol;
+	OutHLSL += TEXT("Buffer<int> ") + BufferName + TEXT(";\n");
+
+	// int MaxNumberOfIndexesPerPoint_XX;
+	BufferName = UNiagaraDataInterfaceHoudiniCSV::MaxNumberOfIndexesPerPointBaseName + ParamInfo.DataInterfaceHLSLSymbol;
+	OutHLSL += TEXT("int ") + BufferName + TEXT(";\n");
+
+	// Buffer<int> PointValueIndexesBuffer_XX;
+	BufferName = UNiagaraDataInterfaceHoudiniCSV::PointValueIndexesBufferBaseName + ParamInfo.DataInterfaceHLSLSymbol;
+	OutHLSL += TEXT("Buffer<int> ") + BufferName + TEXT(";\n");
+
+	// int LastSpawnedPointId_XX;
+	BufferName = UNiagaraDataInterfaceHoudiniCSV::LastSpawnedPointIdBaseName + ParamInfo.DataInterfaceHLSLSymbol;
+	OutHLSL += TEXT("int ") + BufferName + TEXT(";\n");
+}
+
+FRWBuffer& UNiagaraDataInterfaceHoudiniCSV::GetFloatValuesGPUBuffer()
+{
+	//TODO: This isn't really very thread safe. Need to move to a proxy like system where DIs can push data to the RT safely.
+	if ( FloatValuesGPUBufferDirty )
+	{
+		FloatValuesGPUBuffer.Release();
+
+		uint32 NumElements = HoudiniCSVAsset ? HoudiniCSVAsset->FloatCSVData.Num() : 0;
+		FloatValuesGPUBuffer.Initialize( sizeof(float), NumElements, EPixelFormat::PF_R32_FLOAT, BUF_Static);
+
+		uint32 BufferSize = NumElements * sizeof( float );
+		float* BufferData = static_cast<float*>( RHILockVertexBuffer( FloatValuesGPUBuffer.Buffer, 0, BufferSize, EResourceLockMode::RLM_WriteOnly ) );
+
+		if ( HoudiniCSVAsset )
+			FPlatformMemory::Memcpy( BufferData, HoudiniCSVAsset->FloatCSVData.GetData(), BufferSize );
+
+		RHIUnlockVertexBuffer( FloatValuesGPUBuffer.Buffer );
+		FloatValuesGPUBufferDirty = false;
+	}
+
+	return FloatValuesGPUBuffer;
+}
+
+FRWBuffer& UNiagaraDataInterfaceHoudiniCSV::GetSpecialAttributesColumnIndexesGPUBuffer()
+{
+	//TODO: This isn't really very thread safe. Need to move to a proxy like system where DIs can push data to the RT safely.
+	if ( SpecialAttributesColumnIndexesGPUBufferDirty )
+	{
+		SpecialAttributesColumnIndexesGPUBuffer.Release();
+
+		uint32 NumElements = HoudiniCSVAsset ? HoudiniCSVAsset->SpecialAttributesColumnIndexes.Num() : 0;
+		SpecialAttributesColumnIndexesGPUBuffer.Initialize(sizeof(int32), NumElements, EPixelFormat::PF_R32_SINT, BUF_Static);
+
+		uint32 BufferSize = NumElements * sizeof(int32);
+		float* BufferData = static_cast<float*>(RHILockVertexBuffer(SpecialAttributesColumnIndexesGPUBuffer.Buffer, 0, BufferSize, EResourceLockMode::RLM_WriteOnly));
+
+		if (HoudiniCSVAsset)
+			FPlatformMemory::Memcpy(BufferData, HoudiniCSVAsset->SpecialAttributesColumnIndexes.GetData(), BufferSize);
+
+		RHIUnlockVertexBuffer(SpecialAttributesColumnIndexesGPUBuffer.Buffer);
+		SpecialAttributesColumnIndexesGPUBufferDirty = false;
+	}
+
+	return SpecialAttributesColumnIndexesGPUBuffer;
+}
+
+FRWBuffer& UNiagaraDataInterfaceHoudiniCSV::GetSpawnTimesGPUBuffer()
+{
+	//TODO: This isn't really very thread safe. Need to move to a proxy like system where DIs can push data to the RT safely.
+	if ( SpawnTimesGPUBufferDirty )
+	{
+		SpawnTimesGPUBuffer.Release();
+
+		uint32 NumElements = HoudiniCSVAsset ? HoudiniCSVAsset->SpawnTimes.Num() : 0;
+		SpawnTimesGPUBuffer.Initialize( sizeof(float), NumElements, EPixelFormat::PF_R32_FLOAT, BUF_Static);
+				
+		uint32 BufferSize = NumElements * sizeof( float );
+		float* BufferData = static_cast<float*>( RHILockVertexBuffer(SpawnTimesGPUBuffer.Buffer, 0, BufferSize, EResourceLockMode::RLM_WriteOnly ) );
+
+		if ( HoudiniCSVAsset )
+			FPlatformMemory::Memcpy( BufferData, HoudiniCSVAsset->SpawnTimes.GetData(), BufferSize );
+
+		RHIUnlockVertexBuffer( SpawnTimesGPUBuffer.Buffer );
+		SpawnTimesGPUBufferDirty = false;
+	}
+
+	return SpawnTimesGPUBuffer;
+}
+
+FRWBuffer& UNiagaraDataInterfaceHoudiniCSV::GetLifeValuesGPUBuffer()
+{
+	//TODO: This isn't really very thread safe. Need to move to a proxy like system where DIs can push data to the RT safely.
+	if ( LifeValuesGPUBufferDirty )
+	{
+		LifeValuesGPUBuffer.Release();
+
+		uint32 NumElements = HoudiniCSVAsset ? HoudiniCSVAsset->LifeValues.Num() : 0;
+		LifeValuesGPUBuffer.Initialize( sizeof(float), NumElements, EPixelFormat::PF_R32_FLOAT, BUF_Static);
+
+		uint32 BufferSize = NumElements * sizeof( float );
+		float* BufferData = static_cast<float*>( RHILockVertexBuffer( LifeValuesGPUBuffer.Buffer, 0, BufferSize, EResourceLockMode::RLM_WriteOnly ) );
+
+		if ( HoudiniCSVAsset )
+			FPlatformMemory::Memcpy( BufferData, HoudiniCSVAsset->LifeValues.GetData(), BufferSize );
+
+		RHIUnlockVertexBuffer(LifeValuesGPUBuffer.Buffer );
+		LifeValuesGPUBufferDirty = false;
+	}
+
+	return LifeValuesGPUBuffer;
+}
+
+FRWBuffer& UNiagaraDataInterfaceHoudiniCSV::GetPointTypesGPUBuffer()
+{
+	//TODO: This isn't really very thread safe. Need to move to a proxy like system where DIs can push data to the RT safely.
+	if ( PointTypesGPUBufferDirty )
+	{
+		PointTypesGPUBuffer.Release();
+
+		uint32 NumElements = HoudiniCSVAsset ? HoudiniCSVAsset->PointTypes.Num() : 0;
+		PointTypesGPUBuffer.Initialize( sizeof(int32), NumElements, EPixelFormat::PF_R32_SINT, BUF_Static);
+
+		uint32 BufferSize = NumElements * sizeof(int32);
+		int32* BufferData = static_cast<int32*>( RHILockVertexBuffer(PointTypesGPUBuffer.Buffer, 0, BufferSize, EResourceLockMode::RLM_WriteOnly ) );
+
+		if ( HoudiniCSVAsset )
+			FPlatformMemory::Memcpy( BufferData, HoudiniCSVAsset->PointTypes.GetData(), BufferSize );
+
+		RHIUnlockVertexBuffer(PointTypesGPUBuffer.Buffer );
+		PointTypesGPUBufferDirty = false;
+	}
+
+	return PointTypesGPUBuffer;
+}
+
+FRWBuffer& UNiagaraDataInterfaceHoudiniCSV::GetPointValueIndexesGPUBuffer()
+{
+	//TODO: This isn't really very thread safe. Need to move to a proxy like system where DIs can push data to the RT safely.
+	if ( PointValueIndexesGPUBufferDirty )
+	{
+		PointValueIndexesGPUBuffer.Release();
+
+		uint32 NumPoints = HoudiniCSVAsset ? HoudiniCSVAsset->PointValueIndexes.Num() : 0;
+		// Add an extra to the max number so all indexes end by a -1
+		uint32 MaxNumIndexesPerPoint = HoudiniCSVAsset ? HoudiniCSVAsset->GetMaxNumberOfPointValueIndexes() + 1: 0;
+		uint32 NumElements = NumPoints * MaxNumIndexesPerPoint;
+		
+		TArray<int32> PointValueIndexes;
+		PointValueIndexes.Init(-1, NumElements);
+		if ( HoudiniCSVAsset )
+		{
+			// We need to flatten the nested array for HLSL conversion
+			for ( int32 PointID = 0; PointID < HoudiniCSVAsset->PointValueIndexes.Num(); PointID++ )
+			{
+				TArray<int32> RowIndexes = HoudiniCSVAsset->PointValueIndexes[PointID].RowIndexes;
+				for( int32 Idx = 0; Idx < RowIndexes.Num(); Idx++ )
+				{
+					PointValueIndexes[ PointID * MaxNumIndexesPerPoint + Idx ] = RowIndexes[ Idx ];
+				}
+			}
+		}
+
+		PointValueIndexesGPUBuffer.Initialize(sizeof(int32), NumElements, EPixelFormat::PF_R32_SINT, BUF_Static);
+
+		uint32 BufferSize = NumElements * sizeof(int32);
+		int32* BufferData = static_cast<int32*>( RHILockVertexBuffer( PointValueIndexesGPUBuffer.Buffer, 0, BufferSize, EResourceLockMode::RLM_WriteOnly ) );
+		FPlatformMemory::Memcpy( BufferData, PointValueIndexes.GetData(), BufferSize );
+		RHIUnlockVertexBuffer( PointValueIndexesGPUBuffer.Buffer );
+
+		PointValueIndexesGPUBufferDirty = false;
+	}
+
+	return PointValueIndexesGPUBuffer;
+}
+
+// Parameters used for GPU sim compatibility
+struct FNiagaraDataInterfaceParametersCS_HoudiniCSV : public FNiagaraDataInterfaceParametersCS
+{
+	virtual void Bind(const FNiagaraDataInterfaceParamRef& ParamRef, const class FShaderParameterMap& ParameterMap) override
+	{
+		NumberOfRows.Bind(ParameterMap, *(UNiagaraDataInterfaceHoudiniCSV::NumberOfRowsBaseName + ParamRef.ParameterInfo.DataInterfaceHLSLSymbol));
+		NumberOfColumns.Bind(ParameterMap, *(UNiagaraDataInterfaceHoudiniCSV::NumberOfColumnsBaseName + ParamRef.ParameterInfo.DataInterfaceHLSLSymbol));
+		NumberOfPoints.Bind(ParameterMap, *(UNiagaraDataInterfaceHoudiniCSV::NumberOfPointsBaseName + ParamRef.ParameterInfo.DataInterfaceHLSLSymbol));
+		
+		FloatValuesBuffer.Bind(ParameterMap, *(UNiagaraDataInterfaceHoudiniCSV::FloatValuesBufferBaseName + ParamRef.ParameterInfo.DataInterfaceHLSLSymbol));
+
+		SpecialAttributesColumnIndexesBuffer.Bind(ParameterMap, *(UNiagaraDataInterfaceHoudiniCSV::SpecialAttributesColumnIndexesBufferBaseName + ParamRef.ParameterInfo.DataInterfaceHLSLSymbol));		
+
+		SpawnTimesBuffer.Bind(ParameterMap, *(UNiagaraDataInterfaceHoudiniCSV::SpawnTimesBufferBaseName + ParamRef.ParameterInfo.DataInterfaceHLSLSymbol));
+		LifeValuesBuffer.Bind(ParameterMap, *(UNiagaraDataInterfaceHoudiniCSV::LifeValuesBufferBaseName + ParamRef.ParameterInfo.DataInterfaceHLSLSymbol));
+		PointTypesBuffer.Bind(ParameterMap, *(UNiagaraDataInterfaceHoudiniCSV::PointTypesBufferBaseName + ParamRef.ParameterInfo.DataInterfaceHLSLSymbol));
+
+		MaxNumberOfIndexesPerPoint.Bind(ParameterMap, *(UNiagaraDataInterfaceHoudiniCSV::MaxNumberOfIndexesPerPointBaseName + ParamRef.ParameterInfo.DataInterfaceHLSLSymbol));
+		PointValueIndexesBuffer.Bind(ParameterMap, *(UNiagaraDataInterfaceHoudiniCSV::PointValueIndexesBufferBaseName + ParamRef.ParameterInfo.DataInterfaceHLSLSymbol));
+
+		LastSpawnedPointId.Bind(ParameterMap, *(UNiagaraDataInterfaceHoudiniCSV::LastSpawnedPointIdBaseName + ParamRef.ParameterInfo.DataInterfaceHLSLSymbol));
+	}
+
+	virtual void Serialize(FArchive& Ar)override
+	{
+		Ar << NumberOfRows;
+		Ar << NumberOfColumns;
+		Ar << NumberOfPoints;
+
+		Ar << FloatValuesBuffer;
+		Ar << SpecialAttributesColumnIndexesBuffer;
+
+		Ar << SpawnTimesBuffer;
+		Ar << LifeValuesBuffer;
+		Ar << PointTypesBuffer;
+
+		Ar << MaxNumberOfIndexesPerPoint;
+		Ar << PointValueIndexesBuffer;
+
+		Ar << LastSpawnedPointId;
+	}
+
+	virtual void Set( FRHICommandList& RHICmdList, FNiagaraShader* Shader, class UNiagaraDataInterface* DataInterface ) const override
+	{
+		check( IsInRenderingThread() );
+
+		const FComputeShaderRHIParamRef ComputeShaderRHI = Shader->GetComputeShader();
+		UNiagaraDataInterfaceHoudiniCSV* HoudiniDI = CastChecked<UNiagaraDataInterfaceHoudiniCSV>( DataInterface );
+		if ( !HoudiniDI )
+			return;
+
+		SetShaderValue(RHICmdList, ComputeShaderRHI, NumberOfRows, HoudiniDI->GetNumberOfRows());
+		SetShaderValue(RHICmdList, ComputeShaderRHI, NumberOfColumns, HoudiniDI->GetNumberOfColumns());
+		SetShaderValue(RHICmdList, ComputeShaderRHI, NumberOfPoints, HoudiniDI->GetNumberOfPoints());
+
+		FRWBuffer& FloatRWBuffer = HoudiniDI->GetFloatValuesGPUBuffer();
+ 		RHICmdList.SetShaderResourceViewParameter( ComputeShaderRHI, FloatValuesBuffer.GetBaseIndex(), FloatRWBuffer.SRV );
+
+		FRWBuffer& AttributeColumnRWBuffer = HoudiniDI->GetSpecialAttributesColumnIndexesGPUBuffer();
+		RHICmdList.SetShaderResourceViewParameter(ComputeShaderRHI, SpecialAttributesColumnIndexesBuffer.GetBaseIndex(), AttributeColumnRWBuffer.SRV);
+
+		FRWBuffer& SpawnRWBuffer = HoudiniDI->GetSpawnTimesGPUBuffer();
+		RHICmdList.SetShaderResourceViewParameter(ComputeShaderRHI, SpawnTimesBuffer.GetBaseIndex(), SpawnRWBuffer.SRV);
+		FRWBuffer& LifeRWBuffer = HoudiniDI->GetLifeValuesGPUBuffer();
+		RHICmdList.SetShaderResourceViewParameter(ComputeShaderRHI, LifeValuesBuffer.GetBaseIndex(), LifeRWBuffer.SRV);
+		FRWBuffer& TypesRWBuffer = HoudiniDI->GetPointTypesGPUBuffer();
+		RHICmdList.SetShaderResourceViewParameter(ComputeShaderRHI, PointTypesBuffer.GetBaseIndex(), TypesRWBuffer.SRV);
+
+		SetShaderValue(RHICmdList, ComputeShaderRHI, MaxNumberOfIndexesPerPoint, HoudiniDI->GetMaxNumberOfIndexesPerPoints());
+
+		FRWBuffer& PointValuesIndexesRWBuffer = HoudiniDI->GetPointValueIndexesGPUBuffer();
+		RHICmdList.SetShaderResourceViewParameter(ComputeShaderRHI, PointValueIndexesBuffer.GetBaseIndex(), PointValuesIndexesRWBuffer.SRV);
+
+		SetShaderValue(RHICmdList, ComputeShaderRHI, LastSpawnedPointId, -1);
+	}
+
+private:
+
+	FShaderParameter NumberOfRows;
+	FShaderParameter NumberOfColumns;
+	FShaderParameter NumberOfPoints;
+
+	FShaderResourceParameter FloatValuesBuffer;
+	FShaderResourceParameter SpecialAttributesColumnIndexesBuffer;
+
+	FShaderResourceParameter SpawnTimesBuffer;
+	FShaderResourceParameter LifeValuesBuffer;
+	FShaderResourceParameter PointTypesBuffer;
+
+	FShaderParameter MaxNumberOfIndexesPerPoint;
+	FShaderResourceParameter PointValueIndexesBuffer;
+
+	FShaderParameter LastSpawnedPointId;
+};
+
+FNiagaraDataInterfaceParametersCS* UNiagaraDataInterfaceHoudiniCSV::ConstructComputeParameters()const
+{
+	return new FNiagaraDataInterfaceParametersCS_HoudiniCSV();
+}
+
 #undef LOCTEXT_NAMESPACE
diff --git a/Engine/Plugins/FX/HoudiniNiagara/Source/HoudiniNiagara/Public/HoudiniCSV.h b/Engine/Plugins/FX/HoudiniNiagara/Source/HoudiniNiagara/Public/HoudiniCSV.h
--- a/Engine/Plugins/FX/HoudiniNiagara/Source/HoudiniNiagara/Public/HoudiniCSV.h
+++ b/Engine/Plugins/FX/HoudiniNiagara/Source/HoudiniNiagara/Public/HoudiniCSV.h
@@ -147,6 +147,9 @@ class HOUDININIAGARA_API UHoudiniCSV : public UObject
 	// Return a point's type at spawn
 	bool GetPointType(const int32& PointID, int32& Value) const;
 
+	// Returns the maximum number of indexes per point, used for flattening the buffer for HLSL conversion
+	int32 GetMaxNumberOfPointValueIndexes() const;
+
 	//-----------------------------------------------------------------------------------------
 	//  MEMBER VARIABLES
 	//-----------------------------------------------------------------------------------------
diff --git a/Engine/Plugins/FX/HoudiniNiagara/Source/HoudiniNiagaraEditor/Private/HoudiniNiagaraEditor.cpp b/Engine/Plugins/FX/HoudiniNiagara/Source/HoudiniNiagaraEditor/Private/HoudiniNiagaraEditor.cpp
--- a/Engine/Plugins/FX/HoudiniNiagara/Source/HoudiniNiagaraEditor/Private/HoudiniNiagaraEditor.cpp
+++ b/Engine/Plugins/FX/HoudiniNiagara/Source/HoudiniNiagaraEditor/Private/HoudiniNiagaraEditor.cpp
@@ -24,6 +24,8 @@
 #include "HoudiniNiagaraEditor.h"
 #include "HoudiniCSVAssetActions.h"
 #include "AssetRegistryModule.h"
+#include "Styling/SlateStyleRegistry.h"
+#include "Styling/SlateStyle.h"
 
 #define LOCTEXT_NAMESPACE "FHoudiniNiagaraModule"
 
@@ -35,6 +37,34 @@ void FHoudiniNiagaraEditorModule::StartupModule()
     TSharedRef< IAssetTypeActions > HCSVAction = MakeShareable( new FHoudiniCSVAssetActions() );
     AssetTools.RegisterAssetTypeActions( HCSVAction );
     AssetTypeActions.Add( HCSVAction );
+
+	// Create Slate style set.
+	if (!StyleSet.IsValid())
+	{
+		// Create Slate style set.	
+		StyleSet = MakeShareable(new FSlateStyleSet(TEXT("HoudiniNiagaraStyle")));
+		StyleSet->SetContentRoot(FPaths::EngineContentDir() / TEXT("Editor/Slate"));
+		StyleSet->SetCoreContentRoot(FPaths::EngineContentDir() / TEXT("Slate"));
+
+		// Note, these sizes are in Slate Units. Slate Units do NOT have to map to pixels.
+		const FVector2D Icon16x16(16.0f, 16.0f);
+		const FVector2D Icon64x64(64.0f, 64.0f);
+		const FVector2D Icon128x128(128.0f, 128.0f);
+		
+		static FString IconsDir = FPaths::EnginePluginsDir() / TEXT("FX/HoudiniNiagara/Resources/");
+
+		// Register the Asset icon
+		StyleSet->Set(
+			"ClassIcon.HoudiniCSV",
+			new FSlateImageBrush(IconsDir + TEXT("HCSVIcon128.png"), Icon16x16));
+
+		StyleSet->Set(
+			"ClassThumbnail.HoudiniCSV",
+			new FSlateImageBrush(IconsDir + TEXT("HCSVIcon128.png"), Icon64x64));
+
+		// Register Slate style.
+		FSlateStyleRegistry::RegisterSlateStyle(*StyleSet.Get());
+	}
 }
 
 void FHoudiniNiagaraEditorModule::ShutdownModule()
@@ -49,6 +79,16 @@ void FHoudiniNiagaraEditorModule::ShutdownModule()
 
 		AssetTypeActions.Empty();
     }
+
+	// Unregister Slate style set.
+	if (StyleSet.IsValid())
+	{
+		// Unregister Slate style.
+		FSlateStyleRegistry::UnRegisterSlateStyle(*StyleSet.Get());
+
+		ensure(StyleSet.IsUnique());
+		StyleSet.Reset();
+	}
 }
 
 #undef LOCTEXT_NAMESPACE
diff --git a/Engine/Plugins/FX/HoudiniNiagara/Source/HoudiniNiagaraEditor/Public/HoudiniNiagaraEditor.h b/Engine/Plugins/FX/HoudiniNiagara/Source/HoudiniNiagaraEditor/Public/HoudiniNiagaraEditor.h
--- a/Engine/Plugins/FX/HoudiniNiagara/Source/HoudiniNiagaraEditor/Public/HoudiniNiagaraEditor.h
+++ b/Engine/Plugins/FX/HoudiniNiagara/Source/HoudiniNiagaraEditor/Public/HoudiniNiagaraEditor.h
@@ -27,6 +27,7 @@
 #include "Modules/ModuleInterface.h"
 
 class IAssetTypeActions;
+class FSlateStyleSet;
 
 class FHoudiniNiagaraEditorModule : public IModuleInterface
 {
@@ -40,4 +41,7 @@ class FHoudiniNiagaraEditorModule : public IModuleInterface
 
 	/** AssetType actions associated with Houdini CSV assets. **/
 	TArray< TSharedPtr< IAssetTypeActions > > AssetTypeActions;
+
+	/** Slate styleset used by this module. **/
+	TSharedPtr< FSlateStyleSet > StyleSet;
 };
\ No newline at end of file
