Houdini-Niagara:
- Fixed build break due to missing header file
- Houdini CSV asset are now accessible from Blueprints

diff --git a/Engine/Plugins/FX/HoudiniNiagara/Source/HoudiniNiagara/Private/HoudiniCSV.cpp b/Engine/Plugins/FX/HoudiniNiagara/Source/HoudiniNiagara/Private/HoudiniCSV.cpp
--- a/Engine/Plugins/FX/HoudiniNiagara/Source/HoudiniNiagara/Private/HoudiniCSV.cpp
+++ b/Engine/Plugins/FX/HoudiniNiagara/Source/HoudiniNiagara/Private/HoudiniCSV.cpp
@@ -1124,7 +1124,7 @@ bool UHoudiniCSV::GetColumnIndexFromString( const FString& ColumnTitle, int32& C
 }
 
 // Returns the float value at a given point in the CSV file
-bool UHoudiniCSV::GetFloatValue( const int32& rowIndex, const FString& ColumnTitle, float& value ) const
+bool UHoudiniCSV::GetFloatValueForString( const int32& rowIndex, const FString& ColumnTitle, float& value ) const
 {
     int32 ColIndex = -1;
     if ( !GetColumnIndexFromString( ColumnTitle, ColIndex ) )
diff --git a/Engine/Plugins/FX/HoudiniNiagara/Source/HoudiniNiagara/Private/NiagaraDataInterfaceHoudiniCSV.cpp b/Engine/Plugins/FX/HoudiniNiagara/Source/HoudiniNiagara/Private/NiagaraDataInterfaceHoudiniCSV.cpp
--- a/Engine/Plugins/FX/HoudiniNiagara/Source/HoudiniNiagara/Private/NiagaraDataInterfaceHoudiniCSV.cpp
+++ b/Engine/Plugins/FX/HoudiniNiagara/Source/HoudiniNiagara/Private/NiagaraDataInterfaceHoudiniCSV.cpp
@@ -25,6 +25,7 @@
 #include "NiagaraTypes.h"
 #include "Misc/FileHelper.h"
 #include "NiagaraShader.h"
+#include "ShaderParameterUtils.h"
 
 #define LOCTEXT_NAMESPACE "HoudiniNiagaraCSVDataInterface"
 
diff --git a/Engine/Plugins/FX/HoudiniNiagara/Source/HoudiniNiagara/Public/HoudiniCSV.h b/Engine/Plugins/FX/HoudiniNiagara/Source/HoudiniNiagara/Public/HoudiniCSV.h
--- a/Engine/Plugins/FX/HoudiniNiagara/Source/HoudiniNiagara/Public/HoudiniCSV.h
+++ b/Engine/Plugins/FX/HoudiniNiagara/Source/HoudiniNiagara/Public/HoudiniCSV.h
@@ -60,7 +60,7 @@ struct FPointIndexes
 	TArray<int32> RowIndexes;
 };
 
-UCLASS()
+UCLASS(BlueprintType)
 class HOUDININIAGARA_API UHoudiniCSV : public UObject
 {
 	friend class UNiagaraDataInterfaceHoudiniCSV;
@@ -79,72 +79,94 @@ class HOUDININIAGARA_API UHoudiniCSV : public UObject
 	void SetFileName( const FString& TheFilename );
 
 	// Returns the number of points found in the CSV file
+	UFUNCTION(BlueprintCallable, Category = "Houdini Attributes Data")
 	int32 GetNumberOfPoints() const;
 
 	// Returns the number of rows found in the CSV file
+	UFUNCTION(BlueprintCallable, Category = "Houdini Attributes Data")
 	int32 GetNumberOfRows() const ;
 
 	// Returns the number of columns found in the CSV file
+	UFUNCTION(BlueprintCallable, Category = "Houdini Attributes Data")
 	int32 GetNumberOfColumns() const;
 
 	// Return the column index for a specific attribute
+	//UFUNCTION(BlueprintCallable, Category = "Houdini Attributes Data")
 	int32 GetAttributeColumnIndex(const EHoudiniAttributes& Attr) const;
 
 	// Returns if the specific attribute has a valid column index
 	bool IsValidAttributeColumnIndex(const EHoudiniAttributes& Attr) const;
 
 	// Returns the column index for a given string
+	UFUNCTION(BlueprintCallable, Category = "Houdini Attributes Data")
 	bool GetColumnIndexFromString(const FString& ColumnTitle, int32& ColumnIndex) const;
 
 	// Returns the float value at a given point in the CSV file
+	UFUNCTION(BlueprintCallable, Category = "Houdini Attributes Data")
 	bool GetFloatValue( const int32& rowIndex, const int32& colIndex, float& value ) const;
 	// Returns the float value at a given point in the CSV file
-	bool GetFloatValue( const int32& rowIndex, const FString& ColumnTitle, float& value ) const;
+	UFUNCTION(BlueprintCallable, Category = "Houdini Attributes Data")
+	bool GetFloatValueForString( const int32& rowIndex, const FString& ColumnTitle, float& value ) const;
 	/*
-	// Returns the float value at a given point in the CSV file
+	// Returns the string value at a given point in the CSV file
 	bool GetCSVStringValue( const int32& rowIndex, const int32& colIndex, FString& value );
 	// Returns the string value at a given point in the CSV file
 	bool GetCSVStringValue( const int32& rowIndex, const FString& ColumnTitle, FString& value );
 	*/
 	// Returns a Vector3 for a given point in the CSV file
+	UFUNCTION(BlueprintCallable, Category = "Houdini Attributes Data")
 	bool GetVectorValue( const int32& rowIndex, const int32& colIndex, FVector& value, const bool& DoSwap = true, const bool& DoScale = true ) const;
 	
 	// Returns a time value for a given point in the CSV file
+	UFUNCTION(BlueprintCallable, Category = "Houdini Attributes Data")
 	bool GetTimeValue( const int32& rowIndex, float& value ) const;
 	// Returns a Position Vector3 for a given point in the CSV file (converted to unreal's coordinate system)
+	UFUNCTION(BlueprintCallable, Category = "Houdini Attributes Data")
 	bool GetPositionValue( const int32& rowIndex, FVector& value ) const;
 	// Returns a Normal Vector3 for a given point in the CSV file (converted to unreal's coordinate system)
+	UFUNCTION(BlueprintCallable, Category = "Houdini Attributes Data")
 	bool GetNormalValue( const int32& rowIndex, FVector& value ) const;
 	// Returns a Color for a given point in the CSV file
+	UFUNCTION(BlueprintCallable, Category = "Houdini Attributes Data")
 	bool GetColorValue( const int32& rowIndex, FLinearColor& value ) const;
 	// Returns a Velocity Vector3 for a given point in the CSV file
+	UFUNCTION(BlueprintCallable, Category = "Houdini Attributes Data")
 	bool GetVelocityValue(const int32& rowIndex, FVector& value ) const;
 
 	// Get the last row index for a given time value (the row with a time smaller or equal to desiredTime)
 	// If the CSV file doesn't have time informations, returns false and set LastRowIndex to the last row in the file
 	// If desiredTime is smaller than the time value in the first row, LastRowIndex will be set to -1
 	// If desiredTime is higher than the last time value in the last row of the csv file, LastIndex will be set to the last row's index
+	UFUNCTION(BlueprintCallable, Category = "Houdini Attributes Data")
 	bool GetLastRowIndexAtTime( const float& desiredTime, int32& lastRowIndex ) const;
 
 	// Get the last pointID of the points to be spawned at time t
 	// Invalid Index are used to indicate edge cases:
 	// -1 will be returned if there is no points to spawn ( t is smaller than the first point time )
 	// NumberOfRows will be returned if all points in the CSV have been spawned ( t is higher than the last point time )
+	UFUNCTION(BlueprintCallable, Category = "Houdini Attributes Data")
 	bool GetLastPointIDToSpawnAtTime( const float& time, int32& lastIndex ) const;
 
 	// Returns the previous and next indexes for reading the values of a specified point at a given time
+	UFUNCTION(BlueprintCallable, Category = "Houdini Attributes Data")
 	bool GetRowIndexesForPointAtTime(const int32& PointID, const float& desiredTime, int32& PrevIndex, int32& NextIndex, float& PrevWeight) const;
-	// Returns the value for a point at a given time value (linearly interpolated) 
+	// Returns the value for a point at a given time value (linearly interpolated)
+	UFUNCTION(BlueprintCallable, Category = "Houdini Attributes Data")
 	bool GetPointValueAtTime(const int32& PointID, const int32& ColumnIndex, const float& desiredTime, float& Value) const;
-	// Returns the Vector Value for a given point at a given time value (linearly interpolated) 
+	// Returns the Vector Value for a given point at a given time value (linearly interpolated)
+	UFUNCTION(BlueprintCallable, Category = "Houdini Attributes Data")
 	bool GetPointVectorValueAtTime(const int32& PointID, const int32& ColumnIndex, const float& desiredTime, FVector& Vector, const bool& DoSwap, const bool& DoScale) const;
-	// Returns the Position Value for a given point at a given time value (linearly interpolated) 
+	// Returns the Position Value for a given point at a given time value (linearly interpolated)
+	UFUNCTION(BlueprintCallable, Category = "Houdini Attributes Data")
 	bool GetPointPositionAtTime(const int32& PointID, const float& desiredTime, FVector& Vector) const;
 	// Return a given point's life value at spawn
+	UFUNCTION(BlueprintCallable, Category = "Houdini Attributes Data")
 	bool GetPointLife(const int32& PointID, float& Value) const;
 	// Return a point's life for a given time value
+	UFUNCTION(BlueprintCallable, Category = "Houdini Attributes Data")
 	bool GetPointLifeAtTime(const int32& PointID, const float& DesiredTime, float& Value) const;
 	// Return a point's type at spawn
+	UFUNCTION(BlueprintCallable, Category = "Houdini Attributes Data")
 	bool GetPointType(const int32& PointID, int32& Value) const;
 
 	// Returns the maximum number of indexes per point, used for flattening the buffer for HLSL conversion
