Houdini - Niagara:
Added a version to FNiagaraDataInterfaceParametersCS_HoudiniCSV to avoid further issue when deserializing them.
Fixed bug causing no points to spawn when all points in a file spawned at the same time value.
Removed unneeded tests/restriction when calling functions requiring time value on file without time values.
When no time value is in the file, it is cosidered to always be 0.
Fixed differences between GPU/CPU code.
Fixed special attributes column for Vectors in the CSV file not properly named after being expanded.

diff --git a/Engine/Plugins/FX/HoudiniNiagara/Source/HoudiniNiagara/Classes/NiagaraDataInterfaceHoudiniCSV.h b/Engine/Plugins/FX/HoudiniNiagara/Source/HoudiniNiagara/Classes/NiagaraDataInterfaceHoudiniCSV.h
--- a/Engine/Plugins/FX/HoudiniNiagara/Source/HoudiniNiagara/Classes/NiagaraDataInterfaceHoudiniCSV.h
+++ b/Engine/Plugins/FX/HoudiniNiagara/Source/HoudiniNiagara/Classes/NiagaraDataInterfaceHoudiniCSV.h
@@ -177,6 +177,7 @@ public:
 	static const FString MaxNumberOfIndexesPerPointBaseName;
 	static const FString PointValueIndexesBufferBaseName;
 	static const FString LastSpawnedPointIdBaseName;
+	static const FString LastSpawnTimeBaseName;
 
 	// Member variables accessors
 	FORCEINLINE int32 GetNumberOfRows()const { return HoudiniCSVAsset ? HoudiniCSVAsset->GetNumberOfRows() : 0; }
@@ -221,4 +222,7 @@ protected:
 	// Last Spawned PointID
 	UPROPERTY()
 	int32 LastSpawnedPointID;
+	// Last Spawn time
+	UPROPERTY()
+	float LastSpawnTime;
 };
diff --git a/Engine/Plugins/FX/HoudiniNiagara/Source/HoudiniNiagara/Private/HoudiniCSV.cpp b/Engine/Plugins/FX/HoudiniNiagara/Source/HoudiniNiagara/Private/HoudiniCSV.cpp
--- a/Engine/Plugins/FX/HoudiniNiagara/Source/HoudiniNiagara/Private/HoudiniCSV.cpp
+++ b/Engine/Plugins/FX/HoudiniNiagara/Source/HoudiniNiagara/Private/HoudiniCSV.cpp
@@ -34,7 +34,6 @@
 
 DEFINE_LOG_CATEGORY(LogHoudiniNiagara);
 
-
 struct FHoudiniCSVSortPredicate
 {
 	FHoudiniCSVSortPredicate(const int32 &InTimeCol, const int32 &InIDCol )
@@ -310,9 +309,13 @@ bool UHoudiniCSV::UpdateFromStringArray( TArray<FString>& RawStringArray )
 			*/
 		}
 
-		// Look at the point ID, the max ID will be our number of points
-		if ( NumberOfPoints <= CurrentID )
-			NumberOfPoints = CurrentID + 1;
+		// If we dont have Point ID informations, we still want to fill the PointValueIndexes array
+		if ( !IsValidAttributeColumnIndex( EHoudiniAttributes::POINTID ) )
+		{
+			// Each row is considered its own point
+			PointValueIndexes.Add(FPointIndexes());
+			PointValueIndexes[ rowIdx ].RowIndexes.Add( rowIdx );
+		}
     }
 	
 	NumberOfPoints = HoudiniIDToNiagaraIDMap.Num();
@@ -472,18 +475,6 @@ bool UHoudiniCSV::ParseCSVTitleRow( const FString& TitleRow, const FString& Firs
 		}
     }
 
-	// Extract the special attributes column indexes we found
-	int32 PositionColumnIndex = GetAttributeColumnIndex(EHoudiniAttributes::POSITION);
-	int32 NormalColumnIndex = GetAttributeColumnIndex(EHoudiniAttributes::NORMAL);
-	int32 TimeColumnIndex = GetAttributeColumnIndex(EHoudiniAttributes::TIME);
-	int32 IDColumnIndex = GetAttributeColumnIndex(EHoudiniAttributes::POINTID);
-	int32 AliveColumnIndex = GetAttributeColumnIndex(EHoudiniAttributes::ALIVE);
-	int32 LifeColumnIndex = GetAttributeColumnIndex(EHoudiniAttributes::LIFE);
-	int32 ColorColumnIndex = GetAttributeColumnIndex(EHoudiniAttributes::COLOR);
-	int32 AlphaColumnIndex = GetAttributeColumnIndex(EHoudiniAttributes::ALPHA);
-	int32 VelocityColumnIndex = GetAttributeColumnIndex(EHoudiniAttributes::VELOCITY);
-	int32 TypeColumnIndex = GetAttributeColumnIndex(EHoudiniAttributes::TYPE);
-
 	// Read the first row of the CSV file, and look for packed vectors value (X,Y,Z)
     // We'll have to expand them in the title row to match the parsed data
 	HasPackedVectors = false;
@@ -530,6 +521,18 @@ bool UHoudiniCSV::ParseCSVTitleRow( const FString& TitleRow, const FString& Firs
 		// Increment the number of columns
 		NumberOfColumns += ( FoundVectorSize - 1 );
 
+		// Extract the special attributes column indexes we found
+		int32 PositionColumnIndex = GetAttributeColumnIndex(EHoudiniAttributes::POSITION);
+		int32 NormalColumnIndex = GetAttributeColumnIndex(EHoudiniAttributes::NORMAL);
+		int32 TimeColumnIndex = GetAttributeColumnIndex(EHoudiniAttributes::TIME);
+		int32 IDColumnIndex = GetAttributeColumnIndex(EHoudiniAttributes::POINTID);
+		int32 AliveColumnIndex = GetAttributeColumnIndex(EHoudiniAttributes::ALIVE);
+		int32 LifeColumnIndex = GetAttributeColumnIndex(EHoudiniAttributes::LIFE);
+		int32 ColorColumnIndex = GetAttributeColumnIndex(EHoudiniAttributes::COLOR);
+		int32 AlphaColumnIndex = GetAttributeColumnIndex(EHoudiniAttributes::ALPHA);
+		int32 VelocityColumnIndex = GetAttributeColumnIndex(EHoudiniAttributes::VELOCITY);
+		int32 TypeColumnIndex = GetAttributeColumnIndex(EHoudiniAttributes::TYPE);
+
 		// Expand TitleRowArray
 		if ( ( FoundPackedVectorColumnIndex == PositionColumnIndex ) && ( FoundVectorSize == 3 ) )
 		{
@@ -833,6 +836,7 @@ bool UHoudiniCSV::GetLastRowIndexAtTime(const float& desiredTime, int32& lastRow
 // If desiredTime is higher than the last point time in the csv file, LastIndex will be set to the last point's index
 bool UHoudiniCSV::GetLastPointIDToSpawnAtTime( const float& desiredTime, int32& lastID ) const
 {
+	/*
 	// If we dont have proper time info, always return the last point
 	int32 TimeColumnIndex = GetAttributeColumnIndex(EHoudiniAttributes::TIME);    
     if ( TimeColumnIndex < 0 || TimeColumnIndex >= NumberOfColumns )
@@ -840,6 +844,7 @@ bool UHoudiniCSV::GetLastPointIDToSpawnAtTime( const float& desiredTime, int32&
 		lastID = NumberOfPoints - 1;
 		return false;
     }
+	*/
 
     float temp_time = 0.0f;
 	if ( !SpawnTimes.IsValidIndex( NumberOfPoints - 1 ) )
@@ -993,33 +998,6 @@ bool UHoudiniCSV::GetRowIndexesForPointAtTime(const int32& PointID, const float&
 	if ( PointID < 0 || PointID >= NumberOfPoints )
 		return false;
 
-	// Get the spawn time
-	float PointSpawnTime = 0.0f;
-	if ( SpawnTimes.IsValidIndex( PointID ) )
-		PointSpawnTime = SpawnTimes[ PointID ];
-
-	// If the point hasn't spawn before DesiredTime
-	if ( PointSpawnTime > desiredTime )
-		return false;
-
-	// Get the life value
-	float PointLife = 0.0f;
-	if ( LifeValues.IsValidIndex( PointID ) )
-		PointLife = LifeValues[ PointID ];
-
-	// If the point is dead before DesiredTime
-	if ( PointLife > 0.0f && ( PointSpawnTime + PointLife < desiredTime ) )	
-		return false;
-
-	// We don't have Id information
-	// return PointID for the rowIndexes ??
-	if (!IsValidAttributeColumnIndex(EHoudiniAttributes::POINTID))
-		return false;
-
-	// We don't have time information
-	if (!IsValidAttributeColumnIndex(EHoudiniAttributes::TIME))
-		return false;
-
 	// Get the row indexes for this point
 	const TArray<int32>* RowIndexes = nullptr;
 	if ( PointValueIndexes.IsValidIndex( PointID ) )
@@ -1032,33 +1010,33 @@ bool UHoudiniCSV::GetRowIndexesForPointAtTime(const int32& PointID, const float&
 	{
 		// Get the time
 		float currentTime = -1.0f;
-		if ( GetTimeValue( n, currentTime) )
+		if ( !GetTimeValue(n, currentTime) )
+			currentTime = 0.0f;
+
+		if ( currentTime == desiredTime )
+		{
+			PrevIndex = n;
+			NextIndex = n;
+			PrevWeight = 1.0f;
+			return true;
+		}
+		else if ( currentTime < desiredTime )
 		{
-			if ( currentTime == desiredTime )
+			if ( PrevTime == -1.0f || PrevTime < currentTime )
 			{
 				PrevIndex = n;
-				NextIndex = n;
-				PrevWeight = 1.0f;
-				return true;
-			}
-			else if ( currentTime < desiredTime )
-			{
-				if ( PrevTime == -1.0f || PrevTime < currentTime )
-				{
-					PrevIndex = n;
-					PrevTime = currentTime;
-				}
+				PrevTime = currentTime;
 			}
-			else
+		}
+		else
+		{
+			if ( NextTime == -1.0f || NextTime > currentTime)
 			{
-				if ( NextTime == -1.0f || NextTime > currentTime)
-				{
-					NextIndex = n;
-					NextTime = currentTime;
+				NextIndex = n;
+				NextTime = currentTime;
 
-					// TODO: since the csv is sorted by time, we can break now
-					break;
-				}
+				// TODO: since the csv is sorted by time, we can break now
+				break;
 			}
 		}
 	}
@@ -1110,15 +1088,22 @@ bool UHoudiniCSV::IsValidAttributeColumnIndex(const EHoudiniAttributes& Attr) co
 // Returns the column index for a given string
 bool UHoudiniCSV::GetColumnIndexFromString( const FString& ColumnTitle, int32& ColumnIndex ) const
 {
-    if ( !ColumnTitleArray.Find( ColumnTitle, ColumnIndex ) )
+    if ( ColumnTitleArray.Find( ColumnTitle, ColumnIndex ) )
 		return true;
 
-    // Handle packed positions/normals here
-    if ( ColumnTitle.Equals( "P" ) )
+    // Handling special attributes and packed positions/normals
+    if ( ColumnTitle.Equals( "P", ESearchCase::IgnoreCase ) )
 		return ColumnTitleArray.Find( TEXT( "Px" ), ColumnIndex );
-    else if ( ColumnTitle.Equals( "N" ) )
+    else if ( ColumnTitle.Equals( "N", ESearchCase::IgnoreCase ) )
 		return ColumnTitleArray.Find( TEXT( "Nx" ), ColumnIndex );
-	// TODO COLOR etc...
+	else if ( ColumnTitle.Equals("V", ESearchCase::IgnoreCase ) )
+		return ColumnTitleArray.Find( TEXT("Vx"), ColumnIndex);
+	else if ( ColumnTitle.Equals("Cd", ESearchCase::IgnoreCase ) )
+		return ColumnTitleArray.Find( TEXT("R"), ColumnIndex);
+	else if ( ColumnTitle.Equals("Alpha", ESearchCase::IgnoreCase ) )
+		return ColumnTitleArray.Find( TEXT("A"), ColumnIndex);
+	else
+		return ColumnTitleArray.Find( ColumnTitle + TEXT("1"), ColumnIndex );
 
     return false;
 }
diff --git a/Engine/Plugins/FX/HoudiniNiagara/Source/HoudiniNiagara/Private/NiagaraDataInterfaceHoudiniCSV.cpp b/Engine/Plugins/FX/HoudiniNiagara/Source/HoudiniNiagara/Private/NiagaraDataInterfaceHoudiniCSV.cpp
--- a/Engine/Plugins/FX/HoudiniNiagara/Source/HoudiniNiagara/Private/NiagaraDataInterfaceHoudiniCSV.cpp
+++ b/Engine/Plugins/FX/HoudiniNiagara/Source/HoudiniNiagara/Private/NiagaraDataInterfaceHoudiniCSV.cpp
@@ -41,6 +41,7 @@ const FString UNiagaraDataInterfaceHoudiniCSV::PointTypesBufferBaseName(TEXT("Po
 const FString UNiagaraDataInterfaceHoudiniCSV::MaxNumberOfIndexesPerPointBaseName(TEXT("MaxNumberOfIndexesPerPoint_"));
 const FString UNiagaraDataInterfaceHoudiniCSV::PointValueIndexesBufferBaseName(TEXT("PointValueIndexesBuffer_"));
 const FString UNiagaraDataInterfaceHoudiniCSV::LastSpawnedPointIdBaseName(TEXT("LastSpawnedPointId_"));
+const FString UNiagaraDataInterfaceHoudiniCSV::LastSpawnTimeBaseName(TEXT("LastSpawnTime_"));
 
 
 // Name of all the functions available in the data interface
@@ -98,6 +99,7 @@ void UNiagaraDataInterfaceHoudiniCSV::PostInitProperties()
 	PointValueIndexesGPUBufferDirty = true;
 
 	LastSpawnedPointID = -1;
+	LastSpawnTime = -1.0f;
 }
 
 void UNiagaraDataInterfaceHoudiniCSV::PostLoad()
@@ -112,6 +114,7 @@ void UNiagaraDataInterfaceHoudiniCSV::PostLoad()
 	PointValueIndexesGPUBufferDirty = true;
 
 	LastSpawnedPointID = -1;
+	LastSpawnTime = -1.0f;
 }
 
 #if WITH_EDITOR
@@ -133,6 +136,7 @@ void UNiagaraDataInterfaceHoudiniCSV::PostEditChangeProperty(struct FPropertyCha
 			PointValueIndexesGPUBufferDirty = true;
 
 			LastSpawnedPointID = -1;
+			LastSpawnTime = -1.0f;
 		}
     }
 }
@@ -150,6 +154,15 @@ bool UNiagaraDataInterfaceHoudiniCSV::CopyToInternal(UNiagaraDataInterface* Dest
 
     CastedInterface->HoudiniCSVAsset = HoudiniCSVAsset;
 
+	CastedInterface->FloatValuesGPUBufferDirty = true;
+	CastedInterface->SpecialAttributesColumnIndexesGPUBufferDirty = true;
+	CastedInterface->SpawnTimesGPUBufferDirty = true;
+	CastedInterface->LifeValuesGPUBufferDirty = true;
+	CastedInterface->PointTypesGPUBufferDirty = true;
+	CastedInterface->PointValueIndexesGPUBufferDirty = true;
+	CastedInterface->LastSpawnedPointID = -1;
+	CastedInterface->LastSpawnTime = -1.0f;
+
     return true;
 }
 
@@ -983,13 +996,13 @@ void UNiagaraDataInterfaceHoudiniCSV::GetPointIDsToSpawnAtTime( FVectorVMContext
 			{
 				// The CSV file has time informations
 				// First, detect if we need to reset LastSpawnedPointID (after a loop of the emitter)
-				if ( value < LastSpawnedPointID )
+				if ( value < LastSpawnedPointID || t <= LastSpawnTime )
 					LastSpawnedPointID = -1;
 
 				if ( value < 0 )
 				{
 					// Nothing to spawn, t is lower than the point's time
-					LastSpawnedPointID = -1;
+					LastSpawnedPointID = -1;					
 				}
 				else
 				{
@@ -1012,6 +1025,7 @@ void UNiagaraDataInterfaceHoudiniCSV::GetPointIDsToSpawnAtTime( FVectorVMContext
 						count = max - min + 1;
 
 						LastSpawnedPointID = max;
+						LastSpawnTime = t;
 					}
 				}
 			}
@@ -1360,6 +1374,7 @@ bool UNiagaraDataInterfaceHoudiniCSV::GetFunctionHLSL(const FName& DefinitionFun
 	FString MaxNumberOfIndexesPerPointVar = MaxNumberOfIndexesPerPointBaseName + ParamInfo.DataInterfaceHLSLSymbol;
 	FString PointValueIndexesBuffer = PointValueIndexesBufferBaseName + ParamInfo.DataInterfaceHLSLSymbol;
 	FString LastSpawnedPointIDVar = LastSpawnedPointIdBaseName + ParamInfo.DataInterfaceHLSLSymbol;
+	FString LastSpawnTimeVar = LastSpawnTimeBaseName + ParamInfo.DataInterfaceHLSLSymbol;
 
 	// Lambda returning the HLSL code used for reading a Float value in the FloatBuffer
 	auto ReadFloatInBuffer = [&](const FString& OutFloatValue, const FString& FloatRowIndex, const FString& FloatColIndex)
@@ -1402,23 +1417,11 @@ bool UNiagaraDataInterfaceHoudiniCSV::GetFunctionHLSL(const FName& DefinitionFun
 	};
 
 	// Lambda returning the HLSL code for getting the previous and next row indexes for a given particle at a given time	
-	auto GetRowIndexesForPointAtTime = [&](const FString& In_PointID, const FString& In_Time, const FString& Out_PreviousRow, const FString& Out_NextRow, const FString& Out_Weight )
+	auto GetRowIndexesForPointAtTime = [&](const FString& In_PointID, const FString& In_Time, const FString& Out_PreviousRow, const FString& Out_NextRow, const FString& Out_Weight)
 	{
 		FString OutHLSLCode;
 		OutHLSLCode += TEXT("\t// GetRowIndexesForPointAtTime\n");
 		OutHLSLCode += TEXT("\t{\n");
-			// If the point hasn't spawn before DesiredTime
-			OutHLSLCode += TEXT("\t\tfloat spawn_time = ") + SpawnTimeBuffer + TEXT("[ (") + In_PointID + TEXT(") ];\n");
-			OutHLSLCode += TEXT("\t\tif ( spawn_time > (") + In_Time + TEXT(") ){ return; }\n");
-			// If the point is dead before DesiredTime
-			OutHLSLCode += TEXT("\t\tfloat life_value = ") + LifeValuesBuffer + TEXT("[ (") + In_PointID + TEXT(") ];\n");
-			OutHLSLCode += TEXT("\t\tif ( life_value > 0.0f && ( spawn_time + life_value < (") + In_Time + TEXT(") ) ) { return; }\n");
-
-			// Make sure we have time and id attributes
-			OutHLSLCode += TEXT("\t\tint id_col = ") + GetSpecAttributeColumnIndex(EHoudiniAttributes::POINTID) + TEXT(";\n");
-			OutHLSLCode += TEXT("\t\tif ( id_col < 0 || id_col >= ") + NumberOfColumnsVar + TEXT(" ){ return; }\n");
-			OutHLSLCode += TEXT("\t\tint time_col = ") + GetSpecAttributeColumnIndex(EHoudiniAttributes::TIME) + TEXT(";\n");
-			OutHLSLCode += TEXT("\t\tif ( time_col < 0 || time_col >= ") + NumberOfColumnsVar + TEXT(" ){ return; }\n");
 
 			OutHLSLCode += TEXT("\t\tfloat prev_time = -1.0f;\n");
 			OutHLSLCode += TEXT("\t\tfloat next_time = -1.0f;\n");
@@ -1429,7 +1432,11 @@ bool UNiagaraDataInterfaceHoudiniCSV::GetFunctionHLSL(const FName& DefinitionFun
 				OutHLSLCode += TEXT("\t\t\tif ( current_row < 0 ){ break; }\n");
 
 				OutHLSLCode += TEXT("\t\t\tfloat current_time = -1.0f;\n");
-				OutHLSLCode += ReadFloatInBuffer(TEXT("current_time"), TEXT("current_row"), TEXT("time_col"));
+				OutHLSLCode += TEXT("\t\t\tint time_col = ") + GetSpecAttributeColumnIndex(EHoudiniAttributes::TIME) + TEXT(";\n");
+				OutHLSLCode += TEXT("\t\t\tif ( time_col < 0 || time_col >= ") + NumberOfColumnsVar + TEXT(" )\n");
+					OutHLSLCode += TEXT("\t\t\t\t{ current_time = 0.0f; }\n");
+				OutHLSLCode += TEXT("\t\t\telse\n");
+					OutHLSLCode += TEXT("\t\t\t{") + ReadFloatInBuffer(TEXT("current_time"), TEXT("current_row"), TEXT("time_col")) + TEXT(" }\n");
 
 				OutHLSLCode += TEXT("\t\t\tif ( current_time == (") + In_Time + TEXT(") )\n");
 					OutHLSLCode += TEXT("\t\t\t\t{ ") + Out_PreviousRow + TEXT(" = current_row; ") + Out_NextRow + TEXT(" = current_row; ") + Out_Weight + TEXT(" = 1.0; return;}\n");
@@ -1446,7 +1453,7 @@ bool UNiagaraDataInterfaceHoudiniCSV::GetFunctionHLSL(const FName& DefinitionFun
 			OutHLSLCode += TEXT("\t\tif ( ") + Out_NextRow + TEXT(" < 0 )\n");
 				OutHLSLCode += TEXT("\t\t\t{ ") + Out_Weight + TEXT(" = 1.0f; ") + Out_NextRow + TEXT(" = ") + Out_PreviousRow + TEXT("; return;}\n");
 
-			// Calculate the weight
+				// Calculate the weight
 			OutHLSLCode += TEXT("\t\t") + Out_Weight + TEXT(" = ( ( (") + In_Time + TEXT(") - prev_time ) / ( next_time - prev_time ) );\n");
 
 		OutHLSLCode += TEXT("\t}\n");
@@ -1651,7 +1658,7 @@ bool UNiagaraDataInterfaceHoudiniCSV::GetFunctionHLSL(const FName& DefinitionFun
 			OutHLSL += TEXT("\t}\n");
 
 			// Check if we need to reset lastspawnedPointID
-			OutHLSL += TEXT("\tif ( last_id < ") + LastSpawnedPointIDVar + TEXT(")\n");
+			OutHLSL += TEXT("\tif ( last_id < ") + LastSpawnedPointIDVar + TEXT(" || In_Time <= ") + LastSpawnTimeVar + TEXT(")\n");
 				OutHLSL += TEXT("\t\t{") + LastSpawnedPointIDVar + TEXT(" = -1; }\n");
 			// Nothing to spawn, In_Time is lower than the point's time
 			OutHLSL += TEXT("\tif ( last_id < 0 )\n");
@@ -1664,6 +1671,7 @@ bool UNiagaraDataInterfaceHoudiniCSV::GetFunctionHLSL(const FName& DefinitionFun
 			OutHLSL += TEXT("\tOut_MaxIndex = last_id;\n");
 			OutHLSL += TEXT("\tOut_Count = Out_MaxIndex - Out_MinIndex + 1;\n");
 			OutHLSL += TEXT("\t") + LastSpawnedPointIDVar + TEXT(" = Out_MaxIndex;\n");
+			OutHLSL += TEXT("\t") + LastSpawnTimeVar + TEXT(" = In_Time;\n");
 
 		OutHLSL += TEXT("\n}\n");
 		return true;
@@ -1672,8 +1680,8 @@ bool UNiagaraDataInterfaceHoudiniCSV::GetFunctionHLSL(const FName& DefinitionFun
 	{
 		// GetRowIndexesForPointAtTime(int In_PointID, float In_Time, out int Out_PreviousRow, out int Out_NextRow, out float Out_Weight)
 		OutHLSL += TEXT("void ") + InstanceFunctionName + TEXT("(int In_PointID, float In_Time, out int Out_PreviousRow, out int Out_NextRow, out float Out_Weight) \n{\n");
-		OutHLSL += TEXT("\tOut_PreviousRow = 0;\n\tOut_NextRow = 0;\n\tOut_Weight = 0;\n");
 
+			OutHLSL += TEXT("\tOut_PreviousRow = 0;\n\tOut_NextRow = 0;\n\tOut_Weight = 0;\n");
 			OutHLSL += GetRowIndexesForPointAtTime(TEXT("In_PointID"), TEXT("In_Time"), TEXT("Out_PreviousRow"), TEXT("Out_NextRow"), TEXT("Out_Weight"));
 
 		OutHLSL += TEXT("\n}\n");
@@ -1685,7 +1693,7 @@ bool UNiagaraDataInterfaceHoudiniCSV::GetFunctionHLSL(const FName& DefinitionFun
 		OutHLSL += TEXT("void ") + InstanceFunctionName + TEXT("(int In_PointID, float In_Time, out float3 Out_Value) \n{\n");
 		OutHLSL += TEXT("Out_Value = float3( 0.0, 0.0, 0.0 );\n");
 			OutHLSL += TEXT("\tint pos_col = ") + GetSpecAttributeColumnIndex(EHoudiniAttributes::POSITION) + TEXT(";\n");
-			OutHLSL += TEXT("\tint prev_index = -1;int next_index = -1;float weight = 1.0f;\n");
+			OutHLSL += TEXT("\tint prev_index = 0;int next_index = 0;float weight = 0.0f;\n");
 			OutHLSL += GetRowIndexesForPointAtTime(TEXT("In_PointID"), TEXT("In_Time"), TEXT("prev_index"), TEXT("next_index"), TEXT("weight"));
 
 			OutHLSL += TEXT("\tint In_DoSwap = 1;\n");
@@ -1844,6 +1852,10 @@ void UNiagaraDataInterfaceHoudiniCSV::GetParameterDefinitionHLSL(FNiagaraDataInt
 	// int LastSpawnedPointId_XX;
 	BufferName = UNiagaraDataInterfaceHoudiniCSV::LastSpawnedPointIdBaseName + ParamInfo.DataInterfaceHLSLSymbol;
 	OutHLSL += TEXT("int ") + BufferName + TEXT(";\n");
+
+	// int LastSpawnTime_XX;
+	BufferName = UNiagaraDataInterfaceHoudiniCSV::LastSpawnTimeBaseName + ParamInfo.DataInterfaceHLSLSymbol;
+	OutHLSL += TEXT("int ") + BufferName + TEXT(";\n");
 }
 
 FRWBuffer& UNiagaraDataInterfaceHoudiniCSV::GetFloatValuesGPUBuffer()
@@ -1851,9 +1863,11 @@ FRWBuffer& UNiagaraDataInterfaceHoudiniCSV::GetFloatValuesGPUBuffer()
 	//TODO: This isn't really very thread safe. Need to move to a proxy like system where DIs can push data to the RT safely.
 	if ( FloatValuesGPUBufferDirty )
 	{
-		FloatValuesGPUBuffer.Release();
-
 		uint32 NumElements = HoudiniCSVAsset ? HoudiniCSVAsset->FloatCSVData.Num() : 0;
+		if (NumElements <= 0)
+			return FloatValuesGPUBuffer;
+
+		FloatValuesGPUBuffer.Release();
 		FloatValuesGPUBuffer.Initialize( sizeof(float), NumElements, EPixelFormat::PF_R32_FLOAT, BUF_Static);
 
 		uint32 BufferSize = NumElements * sizeof( float );
@@ -1874,9 +1888,11 @@ FRWBuffer& UNiagaraDataInterfaceHoudiniCSV::GetSpecialAttributesColumnIndexesGPU
 	//TODO: This isn't really very thread safe. Need to move to a proxy like system where DIs can push data to the RT safely.
 	if ( SpecialAttributesColumnIndexesGPUBufferDirty )
 	{
-		SpecialAttributesColumnIndexesGPUBuffer.Release();
-
 		uint32 NumElements = HoudiniCSVAsset ? HoudiniCSVAsset->SpecialAttributesColumnIndexes.Num() : 0;
+		if (NumElements <= 0)
+			return SpecialAttributesColumnIndexesGPUBuffer;
+
+		SpecialAttributesColumnIndexesGPUBuffer.Release();
 		SpecialAttributesColumnIndexesGPUBuffer.Initialize(sizeof(int32), NumElements, EPixelFormat::PF_R32_SINT, BUF_Static);
 
 		uint32 BufferSize = NumElements * sizeof(int32);
@@ -1897,9 +1913,11 @@ FRWBuffer& UNiagaraDataInterfaceHoudiniCSV::GetSpawnTimesGPUBuffer()
 	//TODO: This isn't really very thread safe. Need to move to a proxy like system where DIs can push data to the RT safely.
 	if ( SpawnTimesGPUBufferDirty )
 	{
-		SpawnTimesGPUBuffer.Release();
-
 		uint32 NumElements = HoudiniCSVAsset ? HoudiniCSVAsset->SpawnTimes.Num() : 0;
+		if (NumElements <= 0)
+			return SpawnTimesGPUBuffer;
+
+		SpawnTimesGPUBuffer.Release();
 		SpawnTimesGPUBuffer.Initialize( sizeof(float), NumElements, EPixelFormat::PF_R32_FLOAT, BUF_Static);
 				
 		uint32 BufferSize = NumElements * sizeof( float );
@@ -1920,9 +1938,11 @@ FRWBuffer& UNiagaraDataInterfaceHoudiniCSV::GetLifeValuesGPUBuffer()
 	//TODO: This isn't really very thread safe. Need to move to a proxy like system where DIs can push data to the RT safely.
 	if ( LifeValuesGPUBufferDirty )
 	{
-		LifeValuesGPUBuffer.Release();
-
 		uint32 NumElements = HoudiniCSVAsset ? HoudiniCSVAsset->LifeValues.Num() : 0;
+		if (NumElements <= 0)
+			return LifeValuesGPUBuffer;
+
+		LifeValuesGPUBuffer.Release();
 		LifeValuesGPUBuffer.Initialize( sizeof(float), NumElements, EPixelFormat::PF_R32_FLOAT, BUF_Static);
 
 		uint32 BufferSize = NumElements * sizeof( float );
@@ -1943,9 +1963,11 @@ FRWBuffer& UNiagaraDataInterfaceHoudiniCSV::GetPointTypesGPUBuffer()
 	//TODO: This isn't really very thread safe. Need to move to a proxy like system where DIs can push data to the RT safely.
 	if ( PointTypesGPUBufferDirty )
 	{
-		PointTypesGPUBuffer.Release();
-
 		uint32 NumElements = HoudiniCSVAsset ? HoudiniCSVAsset->PointTypes.Num() : 0;
+		if (NumElements <= 0)
+			return PointTypesGPUBuffer;
+
+		PointTypesGPUBuffer.Release();
 		PointTypesGPUBuffer.Initialize( sizeof(int32), NumElements, EPixelFormat::PF_R32_SINT, BUF_Static);
 
 		uint32 BufferSize = NumElements * sizeof(int32);
@@ -1966,13 +1988,14 @@ FRWBuffer& UNiagaraDataInterfaceHoudiniCSV::GetPointValueIndexesGPUBuffer()
 	//TODO: This isn't really very thread safe. Need to move to a proxy like system where DIs can push data to the RT safely.
 	if ( PointValueIndexesGPUBufferDirty )
 	{
-		PointValueIndexesGPUBuffer.Release();
-
 		uint32 NumPoints = HoudiniCSVAsset ? HoudiniCSVAsset->PointValueIndexes.Num() : 0;
 		// Add an extra to the max number so all indexes end by a -1
 		uint32 MaxNumIndexesPerPoint = HoudiniCSVAsset ? HoudiniCSVAsset->GetMaxNumberOfPointValueIndexes() + 1: 0;
 		uint32 NumElements = NumPoints * MaxNumIndexesPerPoint;
-		
+
+		if (NumElements <= 0)
+			return PointValueIndexesGPUBuffer;
+
 		TArray<int32> PointValueIndexes;
 		PointValueIndexes.Init(-1, NumElements);
 		if ( HoudiniCSVAsset )
@@ -1988,6 +2011,7 @@ FRWBuffer& UNiagaraDataInterfaceHoudiniCSV::GetPointValueIndexesGPUBuffer()
 			}
 		}
 
+		PointValueIndexesGPUBuffer.Release();
 		PointValueIndexesGPUBuffer.Initialize(sizeof(int32), NumElements, EPixelFormat::PF_R32_SINT, BUF_Static);
 
 		uint32 BufferSize = NumElements * sizeof(int32);
@@ -2022,10 +2046,13 @@ struct FNiagaraDataInterfaceParametersCS_HoudiniCSV : public FNiagaraDataInterfa
 		PointValueIndexesBuffer.Bind(ParameterMap, *(UNiagaraDataInterfaceHoudiniCSV::PointValueIndexesBufferBaseName + ParamRef.ParameterInfo.DataInterfaceHLSLSymbol));
 
 		LastSpawnedPointId.Bind(ParameterMap, *(UNiagaraDataInterfaceHoudiniCSV::LastSpawnedPointIdBaseName + ParamRef.ParameterInfo.DataInterfaceHLSLSymbol));
+		LastSpawnTime.Bind(ParameterMap, *(UNiagaraDataInterfaceHoudiniCSV::LastSpawnTimeBaseName + ParamRef.ParameterInfo.DataInterfaceHLSLSymbol));
 	}
 
 	virtual void Serialize(FArchive& Ar)override
 	{
+		Ar << Version;
+
 		Ar << NumberOfRows;
 		Ar << NumberOfColumns;
 		Ar << NumberOfPoints;
@@ -2041,6 +2068,7 @@ struct FNiagaraDataInterfaceParametersCS_HoudiniCSV : public FNiagaraDataInterfa
 		Ar << PointValueIndexesBuffer;
 
 		Ar << LastSpawnedPointId;
+		Ar << LastSpawnTime;
 	}
 
 	virtual void Set( FRHICommandList& RHICmdList, FNiagaraShader* Shader, class UNiagaraDataInterface* DataInterface ) const override
@@ -2075,6 +2103,7 @@ struct FNiagaraDataInterfaceParametersCS_HoudiniCSV : public FNiagaraDataInterfa
 		RHICmdList.SetShaderResourceViewParameter(ComputeShaderRHI, PointValueIndexesBuffer.GetBaseIndex(), PointValuesIndexesRWBuffer.SRV);
 
 		SetShaderValue(RHICmdList, ComputeShaderRHI, LastSpawnedPointId, -1);
+		SetShaderValue(RHICmdList, ComputeShaderRHI, LastSpawnTime, -1);
 	}
 
 private:
@@ -2094,6 +2123,9 @@ private:
 	FShaderResourceParameter PointValueIndexesBuffer;
 
 	FShaderParameter LastSpawnedPointId;
+	FShaderParameter LastSpawnTime;
+
+	uint32 Version = 1;
 };
 
 FNiagaraDataInterfaceParametersCS* UNiagaraDataInterfaceHoudiniCSV::ConstructComputeParameters()const
