Houdini - Niagara:
Moved the logic of GetPointIDsToSpawnAtTime to HoudiniCSV  so the function can be called from blueprints.
Renamed the MinIndex/MaxIndex parameters on it to MinId/MaxID.

diff --git a/Engine/Plugins/FX/HoudiniNiagara/Source/HoudiniNiagara/Private/HoudiniCSV.cpp b/Engine/Plugins/FX/HoudiniNiagara/Source/HoudiniNiagara/Private/HoudiniCSV.cpp
--- a/Engine/Plugins/FX/HoudiniNiagara/Source/HoudiniNiagara/Private/HoudiniCSV.cpp
+++ b/Engine/Plugins/FX/HoudiniNiagara/Source/HoudiniNiagara/Private/HoudiniCSV.cpp
@@ -1064,6 +1064,60 @@ bool UHoudiniCSV::GetRowIndexesForPointAtTime(const int32& PointID, const float&
 	return true;
 }
 
+bool UHoudiniCSV::GetPointIDsToSpawnAtTime(
+	const float& desiredTime,
+	int32& MinID, int32& MaxID, int32& Count,
+	int32& LastSpawnedPointID, float& LastSpawnTime ) const
+{
+	int32 lastID = 0;
+	if ( !GetLastPointIDToSpawnAtTime( desiredTime, lastID ) )
+	{
+		// The CSV file doesn't have time informations, so always return all points in the file
+		MinID = 0;
+		MaxID = lastID;
+		Count = MaxID - MinID + 1;
+	}
+	else
+	{
+		// The CSV file has time informations
+		// First, detect if we need to reset LastSpawnedPointID (after a loop of the emitter)
+		if ( lastID < LastSpawnedPointID || desiredTime <= LastSpawnTime )
+			LastSpawnedPointID = -1;
+
+		if ( lastID < 0 )
+		{
+			// Nothing to spawn, t is lower than the point's time
+			LastSpawnedPointID = -1;					
+		}
+		else
+		{
+			// The last time value in the CSV is lower than t, spawn everything if we didnt already!
+			if ( lastID >= GetNumberOfPoints() )
+				lastID = lastID - 1;
+
+			if ( lastID == LastSpawnedPointID )
+			{
+				// We dont have any new point to spawn
+				MinID = lastID;
+				MaxID = lastID;
+				Count = 0;
+			}
+			else
+			{
+				// We have points to spawn at time t
+				MinID = LastSpawnedPointID + 1;
+				MaxID = lastID;
+				Count = MaxID - MinID + 1;
+
+				LastSpawnedPointID = MaxID;
+				LastSpawnTime = desiredTime;
+			}
+		}
+	}
+
+	return true;
+}
+
 int32 UHoudiniCSV::GetAttributeColumnIndex(const EHoudiniAttributes& Attr) const
 {
 	if ( !IsValidAttributeColumnIndex( Attr ) )
diff --git a/Engine/Plugins/FX/HoudiniNiagara/Source/HoudiniNiagara/Private/NiagaraDataInterfaceHoudiniCSV.cpp b/Engine/Plugins/FX/HoudiniNiagara/Source/HoudiniNiagara/Private/NiagaraDataInterfaceHoudiniCSV.cpp
--- a/Engine/Plugins/FX/HoudiniNiagara/Source/HoudiniNiagara/Private/NiagaraDataInterfaceHoudiniCSV.cpp
+++ b/Engine/Plugins/FX/HoudiniNiagara/Source/HoudiniNiagara/Private/NiagaraDataInterfaceHoudiniCSV.cpp
@@ -423,8 +423,8 @@ void UNiagaraDataInterfaceHoudiniCSV::GetFunctions(TArray<FNiagaraFunctionSignat
 		Sig.bRequiresContext = false;
 		Sig.Inputs.Add(FNiagaraVariable( FNiagaraTypeDefinition(GetClass()), TEXT("CSV") ) );				// CSV in
 		Sig.Inputs.Add(FNiagaraVariable( FNiagaraTypeDefinition::GetFloatDef(), TEXT("Time") ) );		    // Time in
-		Sig.Outputs.Add(FNiagaraVariable( FNiagaraTypeDefinition::GetIntDef(), TEXT("MinIndex") ) );	    // Int Out
-		Sig.Outputs.Add(FNiagaraVariable( FNiagaraTypeDefinition::GetIntDef(), TEXT("MaxIndex") ) );	    // Int Out
+		Sig.Outputs.Add(FNiagaraVariable( FNiagaraTypeDefinition::GetIntDef(), TEXT("MinID") ) );			// Int Out
+		Sig.Outputs.Add(FNiagaraVariable( FNiagaraTypeDefinition::GetIntDef(), TEXT("MaxID") ) );			// Int Out
 		Sig.Outputs.Add(FNiagaraVariable( FNiagaraTypeDefinition::GetIntDef(), TEXT("Count") ) );		    // Int Out
 
 		Sig.SetDescription( LOCTEXT( "DataInterfaceHoudini_GetPointIDsToSpawnAtTime",
@@ -983,52 +983,10 @@ void UNiagaraDataInterfaceHoudiniCSV::GetPointIDsToSpawnAtTime( FVectorVMContext
 
 		int32 value = 0;
 		int32 min = 0, max = 0, count = 0;
+
 		if ( HoudiniCSVAsset )
 		{
-			if ( !HoudiniCSVAsset->GetLastPointIDToSpawnAtTime( t, value ) )
-			{
-				// The CSV file doesn't have time informations, so always return all points in the file
-				min = 0;
-				max = value;
-				count = max - min + 1;
-			}
-			else
-			{
-				// The CSV file has time informations
-				// First, detect if we need to reset LastSpawnedPointID (after a loop of the emitter)
-				if ( value < LastSpawnedPointID || t <= LastSpawnTime )
-					LastSpawnedPointID = -1;
-
-				if ( value < 0 )
-				{
-					// Nothing to spawn, t is lower than the point's time
-					LastSpawnedPointID = -1;					
-				}
-				else
-				{
-					// The last time value in the CSV is lower than t, spawn everything if we didnt already!
-					if ( value >= HoudiniCSVAsset->GetNumberOfPoints() )
-						value = value - 1;
-
-					if ( value == LastSpawnedPointID)
-					{
-						// We dont have any new point to spawn
-						min = value;
-						max = value;
-						count = 0;
-					}
-					else
-					{
-						// We have points to spawn at time t
-						min = LastSpawnedPointID + 1;
-						max = value;
-						count = max - min + 1;
-
-						LastSpawnedPointID = max;
-						LastSpawnTime = t;
-					}
-				}
-			}
+			HoudiniCSVAsset->GetPointIDsToSpawnAtTime(t, min, max, count, LastSpawnedPointID, LastSpawnTime);
 		}
 
 		*OutMinValue.GetDest() = min;
@@ -1639,12 +1597,12 @@ bool UNiagaraDataInterfaceHoudiniCSV::GetFunctionHLSL(const FName& DefinitionFun
 	}
 	else if (DefinitionFunctionName == GetPointIDsToSpawnAtTimeName)
 	{
-		// GetPointIDsToSpawnAtTime(float In_Time, out int Out_MinIndex, out int Out_MaxIndex, out int Out_Count)
-		OutHLSL += TEXT("void ") + InstanceFunctionName + TEXT("(float In_Time, out int Out_MinIndex, out int Out_MaxIndex, out int Out_Count) \n{\n");
+		// GetPointIDsToSpawnAtTime(float In_Time, out int Out_MinID, out int Out_MaxID, out int Out_Count)
+		OutHLSL += TEXT("void ") + InstanceFunctionName + TEXT("(float In_Time, out int Out_MinID, out int Out_MaxID, out int Out_Count) \n{\n");
 			
 			OutHLSL += TEXT("\tint time_col = ") + GetSpecAttributeColumnIndex(EHoudiniAttributes::TIME) + TEXT(";\n");
 			OutHLSL += TEXT("\tif( time_col < 0 )\n");
-				OutHLSL += TEXT("\t\t{ Out_Count = ") + NumberOfPointsVar +TEXT("; Out_MinIndex = 0; Out_MaxIndex = Out_Count - 1; return; }\n");
+				OutHLSL += TEXT("\t\t{ Out_Count = ") + NumberOfPointsVar +TEXT("; Out_MinID = 0; Out_MaxID = Out_Count - 1; return; }\n");
 
 			// GetLastPointIDToSpawnAtTime
 			OutHLSL += TEXT("\tint last_id = -1;\n");
@@ -1662,15 +1620,15 @@ bool UNiagaraDataInterfaceHoudiniCSV::GetFunctionHLSL(const FName& DefinitionFun
 				OutHLSL += TEXT("\t\t{") + LastSpawnedPointIDVar + TEXT(" = -1; }\n");
 			// Nothing to spawn, In_Time is lower than the point's time
 			OutHLSL += TEXT("\tif ( last_id < 0 )\n");
-				OutHLSL += TEXT("\t\t{") + LastSpawnedPointIDVar + TEXT(" = -1;Out_Count = 0; Out_MinIndex = 0; Out_MaxIndex = 0;return;}\n");
+				OutHLSL += TEXT("\t\t{") + LastSpawnedPointIDVar + TEXT(" = -1;Out_Count = 0; Out_MinID = 0; Out_MaxID = 0;return;}\n");
 			// We dont have any new point to spawn
 			OutHLSL += TEXT("\tif ( last_id == ") + LastSpawnedPointIDVar + TEXT(")\n");
-				OutHLSL += TEXT("{ Out_MinIndex = last_id; Out_MaxIndex = last_id; Out_Count = 0; return;}\n");
+				OutHLSL += TEXT("{ Out_MinID = last_id; Out_MaxID = last_id; Out_Count = 0; return;}\n");
 
-			OutHLSL += TEXT("\tOut_MinIndex = ") + LastSpawnedPointIDVar + TEXT(" + 1;\n");
-			OutHLSL += TEXT("\tOut_MaxIndex = last_id;\n");
-			OutHLSL += TEXT("\tOut_Count = Out_MaxIndex - Out_MinIndex + 1;\n");
-			OutHLSL += TEXT("\t") + LastSpawnedPointIDVar + TEXT(" = Out_MaxIndex;\n");
+			OutHLSL += TEXT("\tOut_MinID = ") + LastSpawnedPointIDVar + TEXT(" + 1;\n");
+			OutHLSL += TEXT("\tOut_MaxID = last_id;\n");
+			OutHLSL += TEXT("\tOut_Count = Out_MaxID - Out_MinID + 1;\n");
+			OutHLSL += TEXT("\t") + LastSpawnedPointIDVar + TEXT(" = Out_MaxID;\n");
 			OutHLSL += TEXT("\t") + LastSpawnTimeVar + TEXT(" = In_Time;\n");
 
 		OutHLSL += TEXT("\n}\n");
diff --git a/Engine/Plugins/FX/HoudiniNiagara/Source/HoudiniNiagara/Public/HoudiniCSV.h b/Engine/Plugins/FX/HoudiniNiagara/Source/HoudiniNiagara/Public/HoudiniCSV.h
--- a/Engine/Plugins/FX/HoudiniNiagara/Source/HoudiniNiagara/Public/HoudiniCSV.h
+++ b/Engine/Plugins/FX/HoudiniNiagara/Source/HoudiniNiagara/Public/HoudiniCSV.h
@@ -146,6 +146,11 @@ class HOUDININIAGARA_API UHoudiniCSV : public UObject
 	// NumberOfRows will be returned if all points in the CSV have been spawned ( t is higher than the last point time )
 	UFUNCTION(BlueprintCallable, Category = "Houdini Attributes Data")
 	bool GetLastPointIDToSpawnAtTime( const float& time, int32& lastIndex ) const;
+	UFUNCTION(BlueprintCallable, Category = "Houdini Attributes Data")
+	bool GetPointIDsToSpawnAtTime(
+		const float& desiredTime,
+		int32& MinID, int32& MaxID, int32& Count,
+		int32& LastSpawnedPointID, float& LastSpawnTime ) const;
 
 	// Returns the previous and next indexes for reading the values of a specified point at a given time
 	UFUNCTION(BlueprintCallable, Category = "Houdini Attributes Data")
