Houdini-Niagara:
Updated the uplugin file.

diff --git a/Engine/Plugins/FX/HoudiniNiagara/HoudiniNiagara.uplugin b/Engine/Plugins/FX/HoudiniNiagara/HoudiniNiagara.uplugin
--- a/Engine/Plugins/FX/HoudiniNiagara/HoudiniNiagara.uplugin
+++ b/Engine/Plugins/FX/HoudiniNiagara/HoudiniNiagara.uplugin
@@ -3,19 +3,18 @@
 	"Version": 1,
 	"VersionName": "1.0",
 	"FriendlyName": "Houdini Niagara",
-	"Description": "Houdini Niagara Integration",
+	"Description": "Import Houdini point cloud data into UE4 to drive Niagara particle systems.",
 	"Category": "FX",
 	"CreatedBy": "Side Effects Software Inc.",
 	"CreatedByURL": "http://www.sidefx.com",
-	"DocsURL": "",
+	"DocsURL": "https://www.sidefx.com/docs/unreal/_niagara.html",
 	"MarketplaceURL": "",
 	"SupportURL": "",
 	"EnabledByDefault": false,
 	"CanContainContent": true,
 	"IsBetaVersion": true,
 	"Installed": false,
-	"Modules":
-	[
+	"Modules": [
 		{
 			"Name": "HoudiniNiagara",
 			"Type": "Runtime",
@@ -30,8 +29,7 @@
 		}
 	],
 
-	"Plugins":
-	[
+	"Plugins": [
 		{
 			"Name": "Niagara",
 			"Enabled": true
