Houdini-Niagara:
The HoudiniEvent Struct is now registered by the DI.
Removed the HoudiniEvent uasset.

diff --git a/Engine/Plugins/FX/HoudiniNiagara/Source/HoudiniNiagara/Classes/NiagaraDataInterfaceHoudiniCSV.h b/Engine/Plugins/FX/HoudiniNiagara/Source/HoudiniNiagara/Classes/NiagaraDataInterfaceHoudiniCSV.h
--- a/Engine/Plugins/FX/HoudiniNiagara/Source/HoudiniNiagara/Classes/NiagaraDataInterfaceHoudiniCSV.h
+++ b/Engine/Plugins/FX/HoudiniNiagara/Source/HoudiniNiagara/Classes/NiagaraDataInterfaceHoudiniCSV.h
@@ -31,6 +31,69 @@
 #include "NiagaraDataInterface.h"
 #include "NiagaraDataInterfaceHoudiniCSV.generated.h"
 
+USTRUCT()
+struct FHoudiniEvent
+{
+	GENERATED_USTRUCT_BODY()
+
+	UPROPERTY(EditAnywhere, Category = "Houdini Event")
+	FVector Position;
+
+	UPROPERTY(EditAnywhere, Category = "Houdini Event")
+	FVector Normal;
+
+	UPROPERTY(EditAnywhere, Category = "Houdini Event")
+	float Impulse;
+
+	UPROPERTY(EditAnywhere, Category = "Houdini Event")
+	FVector Velocity;
+
+	UPROPERTY(EditAnywhere, Category = "Houdini Event")
+	int32 PointID;
+
+	UPROPERTY(EditAnywhere, Category = "Houdini Event")
+	float Time;	
+
+	UPROPERTY(EditAnywhere, Category = "Houdini Event")
+	float Life;
+
+	UPROPERTY(EditAnywhere, Category = "Houdini Event")
+	FLinearColor Color;
+
+	UPROPERTY(EditAnywhere, Category = "Houdini Event")
+	int32 Type;	
+
+	FHoudiniEvent()
+		: Position(FVector::ZeroVector)
+		, Normal(FVector::ZeroVector)
+		, Impulse(0.0f)
+		, Velocity( FVector::ZeroVector)
+		, PointID(0)
+		, Time(0.0f)
+		, Life(0.0f)
+		, Color(FLinearColor::White)
+		, Type(0)		 
+	{
+	}
+
+	inline bool operator==(const FHoudiniEvent& Other) const
+	{
+		if ( (Other.Position != Position) || (Other.Normal != Normal)
+			|| (Other.Velocity != Velocity) || (Other.Impulse != Impulse)
+			|| (Other.PointID != PointID) || (Other.Time != Time)
+			|| (Other.Life != Life) || (Other.Color != Color)
+			|| (Other.Type != Type) )
+			return false;
+
+		return true;
+	}
+
+	inline bool operator!=(const FHoudiniEvent& Other) const
+	{
+		return !(*this == Other);
+	}
+};
+
 
 /** Data Interface allowing sampling of Houdini CSV files. */
 UCLASS(EditInlineNew, Category = "Houdini Niagara", meta = (DisplayName = "Houdini Array Info"))
diff --git a/Engine/Plugins/FX/HoudiniNiagara/Source/HoudiniNiagara/Private/NiagaraDataInterfaceHoudiniCSV.cpp b/Engine/Plugins/FX/HoudiniNiagara/Source/HoudiniNiagara/Private/NiagaraDataInterfaceHoudiniCSV.cpp
--- a/Engine/Plugins/FX/HoudiniNiagara/Source/HoudiniNiagara/Private/NiagaraDataInterfaceHoudiniCSV.cpp
+++ b/Engine/Plugins/FX/HoudiniNiagara/Source/HoudiniNiagara/Private/NiagaraDataInterfaceHoudiniCSV.cpp
@@ -90,6 +90,7 @@ void UNiagaraDataInterfaceHoudiniCSV::PostInitProperties()
     if (HasAnyFlags(RF_ClassDefaultObject))
     {
 	    FNiagaraTypeRegistry::Register(FNiagaraTypeDefinition(GetClass()), true, false, false);
+		FNiagaraTypeRegistry::Register(FHoudiniEvent::StaticStruct(), true, true, false);
     }
 
 	FloatValuesGPUBufferDirty = true;
