Merging from Dev-Niagara 4360525, 4360676, 4364096 for Simon.Tovey

Some compile time optimizations

- Moved large map of hlsl keyword replacements from C++ to ini file to improve compile times.
- Now not using templates to handle input parameters to data interface. Now taking a bit from the bytecode offset to denote constant vs register.
- As well as improving compile times, this also makes DIs far easier to write and maintain so I've rolled out to all DIs. Has no affect on perf.
- Moved some things into the VectorVM namesspace.

diff --git a/Engine/Plugins/FX/HoudiniNiagara/Source/HoudiniNiagara/Classes/NiagaraDataInterfaceHoudiniCSV.h b/Engine/Plugins/FX/HoudiniNiagara/Source/HoudiniNiagara/Classes/NiagaraDataInterfaceHoudiniCSV.h
--- a/Engine/Plugins/FX/HoudiniNiagara/Source/HoudiniNiagara/Classes/NiagaraDataInterfaceHoudiniCSV.h
+++ b/Engine/Plugins/FX/HoudiniNiagara/Source/HoudiniNiagara/Classes/NiagaraDataInterfaceHoudiniCSV.h
@@ -129,7 +129,6 @@ public:
 	// EXPOSED FUNCTIONS
 
 	// Returns the float value at a given row and column in the CSV file
-	template<typename RowParamType, typename ColParamType>
 	void GetFloatValue(FVectorVMContext& Context);
 
 	/*
@@ -138,39 +137,30 @@ public:
 	*/
 
 	// Returns a Vector3 value for a given row in the CSV file
-	template<typename RowParamType, typename ColParamType>
 	void GetVectorValue(FVectorVMContext& Context);
 
 	// Returns a Vector3 value for a given row in the CSV file
-	template<typename RowParamType, typename ColParamType, typename DoSwapParamType, typename DoScaleParamType>
 	void GetVectorValueEx(FVectorVMContext& Context);
 
 	// Returns the positions for a given row in the CSV file
-	template<typename RowParamType>
 	void GetPosition(FVectorVMContext& Context);
 
 	// Returns the normals for a given row in the CSV file
-	template<typename RowParamType>
 	void GetNormal(FVectorVMContext& Context);
 
 	// Returns the time for a given row in the CSV file
-	template<typename RowParamType>
 	void GetTime(FVectorVMContext& Context);
 
 	// Returns the velocity for a given row in the CSV file
-	template<typename RowParamType>
 	void GetVelocity(FVectorVMContext& Context);
 
 	// Returns the color for a given row in the CSV file
-	template<typename RowParamType>
 	void GetColor(FVectorVMContext& Context);
 
 	// Returns the impulse value for a given row in the CSV file
-	template<typename RowParamType>
 	void GetImpulse(FVectorVMContext& Context);
 
 	// Returns the position and time for a given row in the CSV file
-	template<typename RowParamType>
 	void GetPositionAndTime(FVectorVMContext& Context);
 
 	// Returns the number of rows found in the CSV file
@@ -183,35 +173,27 @@ public:
 	void GetNumberOfPoints(FVectorVMContext& Context);
 
 	// Returns the last index of the points that should be spawned at time t
-	template<typename TimeParamType>
 	void GetLastRowIndexAtTime(FVectorVMContext& Context);
 
 	// Returns the indexes (min, max) and number of points that should be spawned at time t
-	template<typename TimeParamType>
 	void GetPointIDsToSpawnAtTime(FVectorVMContext& Context);
 
 	// Returns the position for a given point at a given time
-	template<typename PointIDParamType, typename TimeParamType>
 	void GetPointPositionAtTime(FVectorVMContext& Context);
 
 	// Returns a float value for a given point at a given time
-	template<typename PointIDParamType, typename ColParamType, typename TimeParamType>
 	void GetPointValueAtTime(FVectorVMContext& Context);
 
 	// Returns a Vector value for a given point at a given time
-	template<typename PointIDParamType, typename ColParamType, typename TimeParamType>
 	void GetPointVectorValueAtTime(FVectorVMContext& Context);
 
 	// Returns a Vector value for a given point at a given time
-	template<typename PointIDParamType, typename ColParamType, typename TimeParamType, typename DoSwapParamType, typename DoScaleParamType>
 	void GetPointVectorValueAtTimeEx(FVectorVMContext& Context);
 
 	// Returns the line indexes (previous, next) for reading values for a given point at a given time
-	template<typename PointIDParamType, typename TimeParamType>
 	void GetRowIndexesForPointAtTime(FVectorVMContext& Context);
 
 	// Return the life value for a given point
-	template<typename PointIDParamType>
 	void GetPointLife(FVectorVMContext& Context);
 
 	// Return the life of a given point at a given time
@@ -219,7 +201,6 @@ public:
 	//void GetPointLifeAtTime(FVectorVMContext& Context);
 
 	// Return the type value for a given point
-	template<typename PointIDParamType>
 	void GetPointType( FVectorVMContext& Context );
 	
 	//----------------------------------------------------------------------------
diff --git a/Engine/Plugins/FX/HoudiniNiagara/Source/HoudiniNiagara/Private/NiagaraDataInterfaceHoudiniCSV.cpp b/Engine/Plugins/FX/HoudiniNiagara/Source/HoudiniNiagara/Private/NiagaraDataInterfaceHoudiniCSV.cpp
--- a/Engine/Plugins/FX/HoudiniNiagara/Source/HoudiniNiagara/Private/NiagaraDataInterfaceHoudiniCSV.cpp
+++ b/Engine/Plugins/FX/HoudiniNiagara/Source/HoudiniNiagara/Private/NiagaraDataInterfaceHoudiniCSV.cpp
@@ -595,72 +595,72 @@ void UNiagaraDataInterfaceHoudiniCSV::GetFunctions(TArray<FNiagaraFunctionSignat
 	}
 }
 
-DEFINE_NDI_RAW_FUNC_BINDER(UNiagaraDataInterfaceHoudiniCSV, GetFloatValue);
-//DEFINE_NDI_RAW_FUNC_BINDER(UNiagaraDataInterfaceHoudiniCSV, GetFloatValueByString);
-DEFINE_NDI_RAW_FUNC_BINDER(UNiagaraDataInterfaceHoudiniCSV, GetVectorValue);
-DEFINE_NDI_RAW_FUNC_BINDER(UNiagaraDataInterfaceHoudiniCSV, GetVectorValueEx);
-DEFINE_NDI_RAW_FUNC_BINDER(UNiagaraDataInterfaceHoudiniCSV, GetPosition);
-DEFINE_NDI_RAW_FUNC_BINDER(UNiagaraDataInterfaceHoudiniCSV, GetNormal);
-DEFINE_NDI_RAW_FUNC_BINDER(UNiagaraDataInterfaceHoudiniCSV, GetTime);
-DEFINE_NDI_RAW_FUNC_BINDER(UNiagaraDataInterfaceHoudiniCSV, GetColor);
-DEFINE_NDI_RAW_FUNC_BINDER(UNiagaraDataInterfaceHoudiniCSV, GetVelocity);
-DEFINE_NDI_RAW_FUNC_BINDER(UNiagaraDataInterfaceHoudiniCSV, GetImpulse);
-DEFINE_NDI_RAW_FUNC_BINDER(UNiagaraDataInterfaceHoudiniCSV, GetPositionAndTime);
-DEFINE_NDI_RAW_FUNC_BINDER(UNiagaraDataInterfaceHoudiniCSV, GetLastRowIndexAtTime);
-DEFINE_NDI_RAW_FUNC_BINDER(UNiagaraDataInterfaceHoudiniCSV, GetPointIDsToSpawnAtTime);
-DEFINE_NDI_RAW_FUNC_BINDER(UNiagaraDataInterfaceHoudiniCSV, GetRowIndexesForPointAtTime);
-DEFINE_NDI_RAW_FUNC_BINDER(UNiagaraDataInterfaceHoudiniCSV, GetPointPositionAtTime);
-DEFINE_NDI_RAW_FUNC_BINDER(UNiagaraDataInterfaceHoudiniCSV, GetPointValueAtTime);
-DEFINE_NDI_RAW_FUNC_BINDER(UNiagaraDataInterfaceHoudiniCSV, GetPointVectorValueAtTime);
-DEFINE_NDI_RAW_FUNC_BINDER(UNiagaraDataInterfaceHoudiniCSV, GetPointVectorValueAtTimeEx);
-DEFINE_NDI_RAW_FUNC_BINDER(UNiagaraDataInterfaceHoudiniCSV, GetPointLife);
-//DEFINE_NDI_RAW_FUNC_BINDER(UNiagaraDataInterfaceHoudiniCSV, GetPointLifeAtTime);
-DEFINE_NDI_RAW_FUNC_BINDER(UNiagaraDataInterfaceHoudiniCSV, GetPointType);
+DEFINE_NDI_DIRECT_FUNC_BINDER(UNiagaraDataInterfaceHoudiniCSV, GetFloatValue);
+//DEFINE_NDI_FUNC_BINDER(UNiagaraDataInterfaceHoudiniCSV, GetFloatValueByString);
+DEFINE_NDI_DIRECT_FUNC_BINDER(UNiagaraDataInterfaceHoudiniCSV, GetVectorValue);
+DEFINE_NDI_DIRECT_FUNC_BINDER(UNiagaraDataInterfaceHoudiniCSV, GetVectorValueEx);
+DEFINE_NDI_DIRECT_FUNC_BINDER(UNiagaraDataInterfaceHoudiniCSV, GetPosition);
+DEFINE_NDI_DIRECT_FUNC_BINDER(UNiagaraDataInterfaceHoudiniCSV, GetNormal);
+DEFINE_NDI_DIRECT_FUNC_BINDER(UNiagaraDataInterfaceHoudiniCSV, GetTime);
+DEFINE_NDI_DIRECT_FUNC_BINDER(UNiagaraDataInterfaceHoudiniCSV, GetColor);
+DEFINE_NDI_DIRECT_FUNC_BINDER(UNiagaraDataInterfaceHoudiniCSV, GetVelocity);
+DEFINE_NDI_DIRECT_FUNC_BINDER(UNiagaraDataInterfaceHoudiniCSV, GetImpulse);
+DEFINE_NDI_DIRECT_FUNC_BINDER(UNiagaraDataInterfaceHoudiniCSV, GetPositionAndTime);
+DEFINE_NDI_DIRECT_FUNC_BINDER(UNiagaraDataInterfaceHoudiniCSV, GetLastRowIndexAtTime);
+DEFINE_NDI_DIRECT_FUNC_BINDER(UNiagaraDataInterfaceHoudiniCSV, GetPointIDsToSpawnAtTime);
+DEFINE_NDI_DIRECT_FUNC_BINDER(UNiagaraDataInterfaceHoudiniCSV, GetRowIndexesForPointAtTime);
+DEFINE_NDI_DIRECT_FUNC_BINDER(UNiagaraDataInterfaceHoudiniCSV, GetPointPositionAtTime);
+DEFINE_NDI_DIRECT_FUNC_BINDER(UNiagaraDataInterfaceHoudiniCSV, GetPointValueAtTime);
+DEFINE_NDI_DIRECT_FUNC_BINDER(UNiagaraDataInterfaceHoudiniCSV, GetPointVectorValueAtTime);
+DEFINE_NDI_DIRECT_FUNC_BINDER(UNiagaraDataInterfaceHoudiniCSV, GetPointVectorValueAtTimeEx);
+DEFINE_NDI_DIRECT_FUNC_BINDER(UNiagaraDataInterfaceHoudiniCSV, GetPointLife);
+//DEFINE_NDI_FUNC_BINDER(UNiagaraDataInterfaceHoudiniCSV, GetPointLifeAtTime);
+DEFINE_NDI_DIRECT_FUNC_BINDER(UNiagaraDataInterfaceHoudiniCSV, GetPointType);
 void UNiagaraDataInterfaceHoudiniCSV::GetVMExternalFunction(const FVMExternalFunctionBindingInfo& BindingInfo, void* InstanceData, FVMExternalFunction &OutFunc)
 {
     if (BindingInfo.Name == GetFloatValueName && BindingInfo.GetNumInputs() == 2 && BindingInfo.GetNumOutputs() == 1)
     {
-		TNDIParamBinder<0, int32, TNDIParamBinder<1, int32, NDI_RAW_FUNC_BINDER(UNiagaraDataInterfaceHoudiniCSV, GetFloatValue)>>::Bind(this, BindingInfo, InstanceData, OutFunc);
+		NDI_FUNC_BINDER(UNiagaraDataInterfaceHoudiniCSV, GetFloatValue)::Bind(this, OutFunc);
     }
     /*else if (BindingInfo.Name == GetFloatValueByStringName && BindingInfo.GetNumInputs() == 2 && BindingInfo.GetNumOutputs() == 1)
     {
-		TNDIParamBinder<0, float, TNDIParamBinder<1, FString, NDI_RAW_FUNC_BINDER(UNiagaraDataInterfaceHoudiniCSV, GetFloatValueByString)>>::Bind(this, BindingInfo, InstanceData, OutFunc);
+		NDI_FUNC_BINDER(UNiagaraDataInterfaceHoudiniCSV, GetFloatValueByString)::Bind(this, BindingInfo, InstanceData, OutFunc);
     }*/
     else if (BindingInfo.Name == GetVectorValueName && BindingInfo.GetNumInputs() == 2 && BindingInfo.GetNumOutputs() == 3)
     {
-		TNDIParamBinder<0, int32, TNDIParamBinder<1, int32, NDI_RAW_FUNC_BINDER(UNiagaraDataInterfaceHoudiniCSV, GetVectorValue)>>::Bind(this, BindingInfo, InstanceData, OutFunc);
+		NDI_FUNC_BINDER(UNiagaraDataInterfaceHoudiniCSV, GetVectorValue)::Bind(this, OutFunc);
     }
 	else if (BindingInfo.Name == GetVectorValueExName && BindingInfo.GetNumInputs() == 4 && BindingInfo.GetNumOutputs() == 3)
 	{
-		TNDIParamBinder<0, int32, TNDIParamBinder<1, int32, TNDIParamBinder<2, bool, TNDIParamBinder<3, bool, NDI_RAW_FUNC_BINDER(UNiagaraDataInterfaceHoudiniCSV, GetVectorValueEx)>>>>::Bind(this, BindingInfo, InstanceData, OutFunc);
+		NDI_FUNC_BINDER(UNiagaraDataInterfaceHoudiniCSV, GetVectorValueEx)::Bind(this, OutFunc);
 	}
     else if (BindingInfo.Name == GetPositionName && BindingInfo.GetNumInputs() == 1 && BindingInfo.GetNumOutputs() == 3)
     {
-		TNDIParamBinder<0, int32, NDI_RAW_FUNC_BINDER(UNiagaraDataInterfaceHoudiniCSV, GetPosition)>::Bind(this, BindingInfo, InstanceData, OutFunc);
+		NDI_FUNC_BINDER(UNiagaraDataInterfaceHoudiniCSV, GetPosition)::Bind(this, OutFunc);
     }
     else if (BindingInfo.Name == GetNormalName && BindingInfo.GetNumInputs() == 1 && BindingInfo.GetNumOutputs() == 3)
     {
-		TNDIParamBinder<0, int32, NDI_RAW_FUNC_BINDER(UNiagaraDataInterfaceHoudiniCSV, GetNormal)>::Bind(this, BindingInfo, InstanceData, OutFunc);
+		NDI_FUNC_BINDER(UNiagaraDataInterfaceHoudiniCSV, GetNormal)::Bind(this, OutFunc);
     }
     else if (BindingInfo.Name == GetTimeName && BindingInfo.GetNumInputs() == 1 && BindingInfo.GetNumOutputs() == 1)
     {
-		TNDIParamBinder<0, int32, NDI_RAW_FUNC_BINDER(UNiagaraDataInterfaceHoudiniCSV, GetTime)>::Bind(this, BindingInfo, InstanceData, OutFunc);
+		NDI_FUNC_BINDER(UNiagaraDataInterfaceHoudiniCSV, GetTime)::Bind(this, OutFunc);
     }
 	else if (BindingInfo.Name == GetVelocityName && BindingInfo.GetNumInputs() == 1 && BindingInfo.GetNumOutputs() == 3)
 	{
-		TNDIParamBinder<0, int32, NDI_RAW_FUNC_BINDER(UNiagaraDataInterfaceHoudiniCSV, GetVelocity)>::Bind(this, BindingInfo, InstanceData, OutFunc);
+		NDI_FUNC_BINDER(UNiagaraDataInterfaceHoudiniCSV, GetVelocity)::Bind(this, OutFunc);
 	}
 	else if (BindingInfo.Name == GetColorName && BindingInfo.GetNumInputs() == 1 && BindingInfo.GetNumOutputs() == 4)
 	{
-		TNDIParamBinder<0, int32, NDI_RAW_FUNC_BINDER(UNiagaraDataInterfaceHoudiniCSV, GetColor)>::Bind(this, BindingInfo, InstanceData, OutFunc);
+		NDI_FUNC_BINDER(UNiagaraDataInterfaceHoudiniCSV, GetColor)::Bind(this, OutFunc);
 	}
 	else if (BindingInfo.Name == GetImpulseName && BindingInfo.GetNumInputs() == 1 && BindingInfo.GetNumOutputs() == 1)
 	{
-		TNDIParamBinder<0, int32, NDI_RAW_FUNC_BINDER(UNiagaraDataInterfaceHoudiniCSV, GetImpulse)>::Bind(this, BindingInfo, InstanceData, OutFunc);
+		NDI_FUNC_BINDER(UNiagaraDataInterfaceHoudiniCSV, GetImpulse)::Bind(this, OutFunc);
 	}
 	else if (BindingInfo.Name == GetPositionAndTimeName && BindingInfo.GetNumInputs() == 1 && BindingInfo.GetNumOutputs() == 4)
     {
-		TNDIParamBinder<0, int32, NDI_RAW_FUNC_BINDER(UNiagaraDataInterfaceHoudiniCSV, GetPositionAndTime)>::Bind(this, BindingInfo, InstanceData, OutFunc);
+		NDI_FUNC_BINDER(UNiagaraDataInterfaceHoudiniCSV, GetPositionAndTime)::Bind(this, OutFunc);
     }
     else if ( BindingInfo.Name == GetNumberOfPointsName && BindingInfo.GetNumInputs() == 0 && BindingInfo.GetNumOutputs() == 1 )
     {
@@ -676,45 +676,45 @@ void UNiagaraDataInterfaceHoudiniCSV::GetVMExternalFunction(const FVMExternalFun
 	}
     else if (BindingInfo.Name == GetLastRowIndexAtTimeName && BindingInfo.GetNumInputs() == 1 && BindingInfo.GetNumOutputs() == 1)
     {
-		TNDIParamBinder<0, float, NDI_RAW_FUNC_BINDER(UNiagaraDataInterfaceHoudiniCSV, GetLastRowIndexAtTime)>::Bind(this, BindingInfo, InstanceData, OutFunc);
+		NDI_FUNC_BINDER(UNiagaraDataInterfaceHoudiniCSV, GetLastRowIndexAtTime)::Bind(this, OutFunc);
     }
     else if (BindingInfo.Name == GetPointIDsToSpawnAtTimeName && BindingInfo.GetNumInputs() == 1 && BindingInfo.GetNumOutputs() == 3)
     {
-		TNDIParamBinder<0, float, NDI_RAW_FUNC_BINDER(UNiagaraDataInterfaceHoudiniCSV, GetPointIDsToSpawnAtTime)>::Bind(this, BindingInfo, InstanceData, OutFunc);
+		NDI_FUNC_BINDER(UNiagaraDataInterfaceHoudiniCSV, GetPointIDsToSpawnAtTime)::Bind(this, OutFunc);
     }
 	else if (BindingInfo.Name == GetRowIndexesForPointAtTimeName && BindingInfo.GetNumInputs() == 2 && BindingInfo.GetNumOutputs() == 3)
 	{
-		TNDIParamBinder<0, int32, TNDIParamBinder<1, float, NDI_RAW_FUNC_BINDER(UNiagaraDataInterfaceHoudiniCSV, GetRowIndexesForPointAtTime)>>::Bind(this, BindingInfo, InstanceData, OutFunc);
+		NDI_FUNC_BINDER(UNiagaraDataInterfaceHoudiniCSV, GetRowIndexesForPointAtTime)::Bind(this, OutFunc);
 	}
 	else if (BindingInfo.Name == GetPointPositionAtTimeName && BindingInfo.GetNumInputs() == 2 && BindingInfo.GetNumOutputs() == 3)
 	{
-		TNDIParamBinder<0, int32, TNDIParamBinder<1, float, NDI_RAW_FUNC_BINDER(UNiagaraDataInterfaceHoudiniCSV, GetPointPositionAtTime)>>::Bind(this, BindingInfo, InstanceData, OutFunc);
+		NDI_FUNC_BINDER(UNiagaraDataInterfaceHoudiniCSV, GetPointPositionAtTime)::Bind(this, OutFunc);
 	}
 	else if (BindingInfo.Name == GetPointValueAtTimeName && BindingInfo.GetNumInputs() == 3 && BindingInfo.GetNumOutputs() == 1)
 	{
-		TNDIParamBinder<0, int32, TNDIParamBinder<1, int32, TNDIParamBinder<2, float, NDI_RAW_FUNC_BINDER(UNiagaraDataInterfaceHoudiniCSV, GetPointValueAtTime)>>>::Bind(this, BindingInfo, InstanceData, OutFunc);
+		NDI_FUNC_BINDER(UNiagaraDataInterfaceHoudiniCSV, GetPointValueAtTime)::Bind(this, OutFunc);
 	}
 	else if (BindingInfo.Name == GetPointVectorValueAtTimeName && BindingInfo.GetNumInputs() == 3 && BindingInfo.GetNumOutputs() == 3)
 	{
-		TNDIParamBinder<0, int32, TNDIParamBinder<1, int32, TNDIParamBinder<2, float, NDI_RAW_FUNC_BINDER(UNiagaraDataInterfaceHoudiniCSV, GetPointVectorValueAtTime)>>>::Bind(this, BindingInfo, InstanceData, OutFunc);
+		NDI_FUNC_BINDER(UNiagaraDataInterfaceHoudiniCSV, GetPointVectorValueAtTime)::Bind(this, OutFunc);
 	}
 	else if (BindingInfo.Name == GetPointVectorValueAtTimeExName && BindingInfo.GetNumInputs() == 5 && BindingInfo.GetNumOutputs() == 3)
 	{
-		TNDIParamBinder<0, int32, TNDIParamBinder<1, int32, TNDIParamBinder<2, float, TNDIParamBinder<3, bool, TNDIParamBinder<4, bool, NDI_RAW_FUNC_BINDER(UNiagaraDataInterfaceHoudiniCSV, GetPointVectorValueAtTimeEx)>>>>>::Bind(this, BindingInfo, InstanceData, OutFunc);
+		NDI_FUNC_BINDER(UNiagaraDataInterfaceHoudiniCSV, GetPointVectorValueAtTimeEx)::Bind(this, OutFunc);
 	}
 	else if (BindingInfo.Name == GetPointLifeName && BindingInfo.GetNumInputs() == 1 && BindingInfo.GetNumOutputs() == 1)
 	{
-		TNDIParamBinder<0, int32, NDI_RAW_FUNC_BINDER(UNiagaraDataInterfaceHoudiniCSV, GetPointLife)>::Bind(this, BindingInfo, InstanceData, OutFunc);
+		NDI_FUNC_BINDER(UNiagaraDataInterfaceHoudiniCSV, GetPointLife)::Bind(this, OutFunc);
 	}
 	/*
 	else if (BindingInfo.Name == GetPointLifeAtTimeName && BindingInfo.GetNumInputs() == 2 && BindingInfo.GetNumOutputs() == 1)
 	{
-		TNDIParamBinder<0, int32, TNDIParamBinder<1, float, NDI_RAW_FUNC_BINDER(UNiagaraDataInterfaceHoudiniCSV, GetPointLifeAtTime)>>::Bind(this, BindingInfo, InstanceData, OutFunc);
+		NDI_FUNC_BINDER(UNiagaraDataInterfaceHoudiniCSV, GetPointLifeAtTime)::Bind(this, BindingInfo, InstanceData, OutFunc);
 	}
 	*/
 	else if (BindingInfo.Name == GetPointTypeName && BindingInfo.GetNumInputs() == 1 && BindingInfo.GetNumOutputs() == 1)
 	{
-		TNDIParamBinder<0, int32, NDI_RAW_FUNC_BINDER(UNiagaraDataInterfaceHoudiniCSV, GetPointType)>::Bind(this, BindingInfo, InstanceData, OutFunc);
+		NDI_FUNC_BINDER(UNiagaraDataInterfaceHoudiniCSV, GetPointType)::Bind(this, OutFunc);
 	}
     else
     {
@@ -725,13 +725,12 @@ void UNiagaraDataInterfaceHoudiniCSV::GetVMExternalFunction(const FVMExternalFun
     }
 }
 
-template<typename RowParamType, typename ColParamType>
 void UNiagaraDataInterfaceHoudiniCSV::GetFloatValue(FVectorVMContext& Context)
 {
-    RowParamType RowParam(Context);
-    ColParamType ColParam(Context);
+    VectorVM::FExternalFuncInputHandler<int32> RowParam(Context);
+    VectorVM::FExternalFuncInputHandler<int32> ColParam(Context);
 
-    FRegisterHandler<float> OutValue(Context);
+    VectorVM::FExternalFuncRegisterHandler<float> OutValue(Context);
 
     for ( int32 i = 0; i < Context.NumInstances; ++i )
     {
@@ -749,15 +748,14 @@ void UNiagaraDataInterfaceHoudiniCSV::GetFloatValue(FVectorVMContext& Context)
     }
 }
 
-template<typename RowParamType, typename ColParamType>
 void UNiagaraDataInterfaceHoudiniCSV::GetVectorValue( FVectorVMContext& Context )
 {
-    RowParamType RowParam(Context);
-    ColParamType ColParam(Context);
+    VectorVM::FExternalFuncInputHandler<int32> RowParam(Context);
+    VectorVM::FExternalFuncInputHandler<int32> ColParam(Context);
 
-	FRegisterHandler<float> OutVectorX(Context);
-	FRegisterHandler<float> OutVectorY(Context);
-	FRegisterHandler<float> OutVectorZ(Context);
+	VectorVM::FExternalFuncRegisterHandler<float> OutVectorX(Context);
+	VectorVM::FExternalFuncRegisterHandler<float> OutVectorY(Context);
+	VectorVM::FExternalFuncRegisterHandler<float> OutVectorZ(Context);
 
     for (int32 i = 0; i < Context.NumInstances; ++i)
     {
@@ -780,25 +778,24 @@ void UNiagaraDataInterfaceHoudiniCSV::GetVectorValue( FVectorVMContext& Context
     }
 }
 
-template<typename RowParamType, typename ColParamType, typename DoSwapParamType, typename DoScaleParamType >
 void UNiagaraDataInterfaceHoudiniCSV::GetVectorValueEx(FVectorVMContext& Context)
 {
-	RowParamType RowParam(Context);
-	ColParamType ColParam(Context);
-	DoSwapParamType DoSwapParam(Context);
-	DoScaleParamType DoScaleParam(Context);
+	VectorVM::FExternalFuncInputHandler<int32> RowParam(Context);
+	VectorVM::FExternalFuncInputHandler<int32> ColParam(Context);
+	VectorVM::FExternalFuncInputHandler<FNiagaraBool> DoSwapParam(Context);
+	VectorVM::FExternalFuncInputHandler<FNiagaraBool> DoScaleParam(Context);
 
-	FRegisterHandler<float> OutVectorX(Context);
-	FRegisterHandler<float> OutVectorY(Context);
-	FRegisterHandler<float> OutVectorZ(Context);
+	VectorVM::FExternalFuncRegisterHandler<float> OutVectorX(Context);
+	VectorVM::FExternalFuncRegisterHandler<float> OutVectorY(Context);
+	VectorVM::FExternalFuncRegisterHandler<float> OutVectorZ(Context);
 
 	for (int32 i = 0; i < Context.NumInstances; ++i)
 	{
 		int32 row = RowParam.Get();
 		int32 col = ColParam.Get();
 
-		bool DoSwap = DoSwapParam.Get();
-		bool DoScale = DoScaleParam.Get();
+		bool DoSwap = DoSwapParam.Get().GetValue();
+		bool DoScale = DoScaleParam.Get().GetValue();
 
 		FVector V = FVector::ZeroVector;
 		if (HoudiniCSVAsset)
@@ -818,13 +815,12 @@ void UNiagaraDataInterfaceHoudiniCSV::GetVectorValueEx(FVectorVMContext& Context
 	}
 }
 /*
-template<typename RowParamType, typename ColTitleParamType>
 void UNiagaraDataInterfaceHoudiniCSV::GetFloatValueByString(FVectorVMContext& Context)
 {
-    RowParamType RowParam(Context);
+    VectorVM::FExternalFuncInputHandler<int32> RowParam(Context);
     ColTitleParamType ColTitleParam(Context);
 
-    FRegisterHandler<float> OutValue(Context);
+    VectorVM::FExternalFuncRegisterHandler<float> OutValue(Context);
 
     for ( int32 i = 0; i < Context.NumInstances; ++i )
     {
@@ -843,13 +839,12 @@ void UNiagaraDataInterfaceHoudiniCSV::GetFloatValueByString(FVectorVMContext& Co
 }
 */
 
-template<typename RowParamType>
 void UNiagaraDataInterfaceHoudiniCSV::GetPosition(FVectorVMContext& Context)
 {
-	RowParamType RowParam(Context);
-    FRegisterHandler<float> OutSampleX(Context);
-    FRegisterHandler<float> OutSampleY(Context);
-    FRegisterHandler<float> OutSampleZ(Context);
+	VectorVM::FExternalFuncInputHandler<int32> RowParam(Context);
+    VectorVM::FExternalFuncRegisterHandler<float> OutSampleX(Context);
+    VectorVM::FExternalFuncRegisterHandler<float> OutSampleY(Context);
+    VectorVM::FExternalFuncRegisterHandler<float> OutSampleZ(Context);
 
     for (int32 i = 0; i < Context.NumInstances; ++i)
     {
@@ -869,13 +864,12 @@ void UNiagaraDataInterfaceHoudiniCSV::GetPosition(FVectorVMContext& Context)
     }
 }
 
-template<typename RowParamType>
 void UNiagaraDataInterfaceHoudiniCSV::GetNormal(FVectorVMContext& Context)
 {
-	RowParamType RowParam(Context);
-    FRegisterHandler<float> OutSampleX(Context);
-    FRegisterHandler<float> OutSampleY(Context);
-    FRegisterHandler<float> OutSampleZ(Context);
+	VectorVM::FExternalFuncInputHandler<int32> RowParam(Context);
+    VectorVM::FExternalFuncRegisterHandler<float> OutSampleX(Context);
+    VectorVM::FExternalFuncRegisterHandler<float> OutSampleY(Context);
+    VectorVM::FExternalFuncRegisterHandler<float> OutSampleZ(Context);
 
     for (int32 i = 0; i < Context.NumInstances; ++i)
     {
@@ -895,11 +889,10 @@ void UNiagaraDataInterfaceHoudiniCSV::GetNormal(FVectorVMContext& Context)
     }
 }
 
-template<typename RowParamType>
 void UNiagaraDataInterfaceHoudiniCSV::GetTime(FVectorVMContext& Context)
 {
-	RowParamType RowParam(Context);
-    FRegisterHandler<float> OutValue(Context);
+	VectorVM::FExternalFuncInputHandler<int32> RowParam(Context);
+    VectorVM::FExternalFuncRegisterHandler<float> OutValue(Context);
 
     for (int32 i = 0; i < Context.NumInstances; ++i)
     {
@@ -915,13 +908,12 @@ void UNiagaraDataInterfaceHoudiniCSV::GetTime(FVectorVMContext& Context)
     }
 }
 
-template<typename RowParamType>
 void UNiagaraDataInterfaceHoudiniCSV::GetVelocity(FVectorVMContext& Context)
 {
-	RowParamType RowParam(Context);
-	FRegisterHandler<float> OutSampleX(Context);
-	FRegisterHandler<float> OutSampleY(Context);
-	FRegisterHandler<float> OutSampleZ(Context);
+	VectorVM::FExternalFuncInputHandler<int32> RowParam(Context);
+	VectorVM::FExternalFuncRegisterHandler<float> OutSampleX(Context);
+	VectorVM::FExternalFuncRegisterHandler<float> OutSampleY(Context);
+	VectorVM::FExternalFuncRegisterHandler<float> OutSampleZ(Context);
 
 	for (int32 i = 0; i < Context.NumInstances; ++i)
 	{
@@ -941,14 +933,13 @@ void UNiagaraDataInterfaceHoudiniCSV::GetVelocity(FVectorVMContext& Context)
 	}
 }
 
-template<typename RowParamType>
 void UNiagaraDataInterfaceHoudiniCSV::GetColor(FVectorVMContext& Context)
 {
-	RowParamType RowParam(Context);
-	FRegisterHandler<float> OutSampleR(Context);
-	FRegisterHandler<float> OutSampleG(Context);
-	FRegisterHandler<float> OutSampleB(Context);
-	FRegisterHandler<float> OutSampleA(Context);
+	VectorVM::FExternalFuncInputHandler<int32> RowParam(Context);
+	VectorVM::FExternalFuncRegisterHandler<float> OutSampleR(Context);
+	VectorVM::FExternalFuncRegisterHandler<float> OutSampleG(Context);
+	VectorVM::FExternalFuncRegisterHandler<float> OutSampleB(Context);
+	VectorVM::FExternalFuncRegisterHandler<float> OutSampleA(Context);
 
 	for (int32 i = 0; i < Context.NumInstances; ++i)
 	{
@@ -970,11 +961,10 @@ void UNiagaraDataInterfaceHoudiniCSV::GetColor(FVectorVMContext& Context)
 	}
 }
 
-template<typename RowParamType>
 void UNiagaraDataInterfaceHoudiniCSV::GetImpulse(FVectorVMContext& Context)
 {
-	RowParamType RowParam(Context);
-	FRegisterHandler<float> OutValue(Context);
+	VectorVM::FExternalFuncInputHandler<int32> RowParam(Context);
+	VectorVM::FExternalFuncRegisterHandler<float> OutValue(Context);
 
 	for (int32 i = 0; i < Context.NumInstances; ++i)
 	{
@@ -991,11 +981,10 @@ void UNiagaraDataInterfaceHoudiniCSV::GetImpulse(FVectorVMContext& Context)
 }
 
 // Returns the last index of the points that should be spawned at time t
-template<typename TimeParamType>
 void UNiagaraDataInterfaceHoudiniCSV::GetLastRowIndexAtTime(FVectorVMContext& Context)
 {
-    TimeParamType TimeParam(Context);
-    FRegisterHandler<int32> OutValue(Context);
+    VectorVM::FExternalFuncInputHandler<float> TimeParam(Context);
+    VectorVM::FExternalFuncRegisterHandler<int32> OutValue(Context);
 
     for (int32 i = 0; i < Context.NumInstances; ++i)
     {
@@ -1012,13 +1001,12 @@ void UNiagaraDataInterfaceHoudiniCSV::GetLastRowIndexAtTime(FVectorVMContext& Co
 }
 
 // Returns the last index of the points that should be spawned at time t
-template<typename TimeParamType>
 void UNiagaraDataInterfaceHoudiniCSV::GetPointIDsToSpawnAtTime( FVectorVMContext& Context )
 {
-    TimeParamType TimeParam( Context );
-    FRegisterHandler<int32> OutMinValue( Context );
-    FRegisterHandler<int32> OutMaxValue( Context );
-    FRegisterHandler<int32> OutCountValue( Context );
+    VectorVM::FExternalFuncInputHandler<float> TimeParam( Context );
+    VectorVM::FExternalFuncRegisterHandler<int32> OutMinValue( Context );
+    VectorVM::FExternalFuncRegisterHandler<int32> OutMaxValue( Context );
+    VectorVM::FExternalFuncRegisterHandler<int32> OutCountValue( Context );
 
     for (int32 i = 0; i < Context.NumInstances; ++i)
     {
@@ -1043,15 +1031,14 @@ void UNiagaraDataInterfaceHoudiniCSV::GetPointIDsToSpawnAtTime( FVectorVMContext
     }
 }
 
-template<typename RowParamType>
 void UNiagaraDataInterfaceHoudiniCSV::GetPositionAndTime(FVectorVMContext& Context)
 {
-    RowParamType RowParam(Context);
+    VectorVM::FExternalFuncInputHandler<int32> RowParam(Context);
 
-    FRegisterHandler<float> OutPosX(Context);
-    FRegisterHandler<float> OutPosY(Context);
-    FRegisterHandler<float> OutPosZ(Context);
-    FRegisterHandler<float> OutTime(Context);
+    VectorVM::FExternalFuncRegisterHandler<float> OutPosX(Context);
+    VectorVM::FExternalFuncRegisterHandler<float> OutPosY(Context);
+    VectorVM::FExternalFuncRegisterHandler<float> OutPosZ(Context);
+    VectorVM::FExternalFuncRegisterHandler<float> OutTime(Context);
 
     for (int32 i = 0; i < Context.NumInstances; ++i)
     {
@@ -1079,15 +1066,14 @@ void UNiagaraDataInterfaceHoudiniCSV::GetPositionAndTime(FVectorVMContext& Conte
     }
 }
 
-template<typename PointIDParamType, typename TimeParamType>
 void UNiagaraDataInterfaceHoudiniCSV::GetRowIndexesForPointAtTime(FVectorVMContext& Context)
 {
-	PointIDParamType PointIDParam(Context);
-	TimeParamType TimeParam(Context);
+	VectorVM::FExternalFuncInputHandler<int32> PointIDParam(Context);
+	VectorVM::FExternalFuncInputHandler<float> TimeParam(Context);
 
-	FRegisterHandler<int32> OutPrevIndex(Context);
-	FRegisterHandler<int32> OutNextIndex(Context);
-	FRegisterHandler<float> OutWeightValue(Context);
+	VectorVM::FExternalFuncRegisterHandler<int32> OutPrevIndex(Context);
+	VectorVM::FExternalFuncRegisterHandler<int32> OutNextIndex(Context);
+	VectorVM::FExternalFuncRegisterHandler<float> OutWeightValue(Context);
 
 	for (int32 i = 0; i < Context.NumInstances; ++i)
     {
@@ -1114,15 +1100,14 @@ void UNiagaraDataInterfaceHoudiniCSV::GetRowIndexesForPointAtTime(FVectorVMConte
     }
 }
 
-template<typename PointIDParamType, typename TimeParamType>
 void UNiagaraDataInterfaceHoudiniCSV::GetPointPositionAtTime( FVectorVMContext& Context )
 {
-	PointIDParamType PointIDParam(Context);
-	TimeParamType TimeParam(Context);
+	VectorVM::FExternalFuncInputHandler<int32> PointIDParam(Context);
+	VectorVM::FExternalFuncInputHandler<float> TimeParam(Context);
 
-	FRegisterHandler<float> OutPosX(Context);
-	FRegisterHandler<float> OutPosY(Context);
-	FRegisterHandler<float> OutPosZ(Context);
+	VectorVM::FExternalFuncRegisterHandler<float> OutPosX(Context);
+	VectorVM::FExternalFuncRegisterHandler<float> OutPosY(Context);
+	VectorVM::FExternalFuncRegisterHandler<float> OutPosZ(Context);
 
 	for (int32 i = 0; i < Context.NumInstances; ++i)
     {
@@ -1147,14 +1132,13 @@ void UNiagaraDataInterfaceHoudiniCSV::GetPointPositionAtTime( FVectorVMContext&
     }
 }
 
-template<typename PointIDParamType, typename ColParamType, typename TimeParamType>
 void UNiagaraDataInterfaceHoudiniCSV::GetPointValueAtTime(FVectorVMContext& Context)
 {
-	PointIDParamType PointIDParam(Context);
-	TimeParamType TimeParam(Context);
-	ColParamType ColParam(Context);
+	VectorVM::FExternalFuncInputHandler<int32> PointIDParam(Context);
+	VectorVM::FExternalFuncInputHandler<float> TimeParam(Context);
+	VectorVM::FExternalFuncInputHandler<int32> ColParam(Context);
 
-	FRegisterHandler<float> OutValue(Context);
+	VectorVM::FExternalFuncRegisterHandler<float> OutValue(Context);
 
 	for (int32 i = 0; i < Context.NumInstances; ++i)
 	{
@@ -1178,16 +1162,15 @@ void UNiagaraDataInterfaceHoudiniCSV::GetPointValueAtTime(FVectorVMContext& Cont
 	}
 }
 
-template<typename PointIDParamType, typename ColParamType, typename TimeParamType>
 void UNiagaraDataInterfaceHoudiniCSV::GetPointVectorValueAtTime(FVectorVMContext& Context)
 {
-	PointIDParamType PointIDParam(Context);
-	ColParamType ColParam(Context);
-	TimeParamType TimeParam(Context);	
+	VectorVM::FExternalFuncInputHandler<int32> PointIDParam(Context);
+	VectorVM::FExternalFuncInputHandler<int32> ColParam(Context);
+	VectorVM::FExternalFuncInputHandler<float> TimeParam(Context);	
 
-	FRegisterHandler<float> OutPosX(Context);
-	FRegisterHandler<float> OutPosY(Context);
-	FRegisterHandler<float> OutPosZ(Context);
+	VectorVM::FExternalFuncRegisterHandler<float> OutPosX(Context);
+	VectorVM::FExternalFuncRegisterHandler<float> OutPosY(Context);
+	VectorVM::FExternalFuncRegisterHandler<float> OutPosZ(Context);
 
 	for (int32 i = 0; i < Context.NumInstances; ++i)
 	{
@@ -1215,18 +1198,17 @@ void UNiagaraDataInterfaceHoudiniCSV::GetPointVectorValueAtTime(FVectorVMContext
 	}
 }
 
-template<typename PointIDParamType, typename ColParamType, typename TimeParamType, typename DoSwapParamType, typename DoScaleParamType >
 void UNiagaraDataInterfaceHoudiniCSV::GetPointVectorValueAtTimeEx(FVectorVMContext& Context)
 {
-	PointIDParamType PointIDParam(Context);
-	ColParamType ColParam(Context);
-	TimeParamType TimeParam(Context);
-	DoSwapParamType DoSwapParam(Context);
-	DoScaleParamType DoScaleParam(Context);
+	VectorVM::FExternalFuncInputHandler<int32> PointIDParam(Context);
+	VectorVM::FExternalFuncInputHandler<int32> ColParam(Context);
+	VectorVM::FExternalFuncInputHandler<float> TimeParam(Context);
+	VectorVM::FExternalFuncInputHandler<FNiagaraBool> DoSwapParam(Context);
+	VectorVM::FExternalFuncInputHandler<FNiagaraBool> DoScaleParam(Context);
 
-	FRegisterHandler<float> OutPosX(Context);
-	FRegisterHandler<float> OutPosY(Context);
-	FRegisterHandler<float> OutPosZ(Context);
+	VectorVM::FExternalFuncRegisterHandler<float> OutPosX(Context);
+	VectorVM::FExternalFuncRegisterHandler<float> OutPosY(Context);
+	VectorVM::FExternalFuncRegisterHandler<float> OutPosZ(Context);
 
 	for (int32 i = 0; i < Context.NumInstances; ++i)
 	{
@@ -1234,8 +1216,8 @@ void UNiagaraDataInterfaceHoudiniCSV::GetPointVectorValueAtTimeEx(FVectorVMConte
 		int32 Col = ColParam.Get();
 		float time = TimeParam.Get();
 
-		bool DoSwap = DoSwapParam.Get();
-		bool DoScale = DoScaleParam.Get();
+		bool DoSwap = DoSwapParam.Get().GetValue();
+		bool DoScale = DoScaleParam.Get().GetValue();
 
 		FVector posVector = FVector::ZeroVector;
 		if (HoudiniCSVAsset)
@@ -1259,12 +1241,11 @@ void UNiagaraDataInterfaceHoudiniCSV::GetPointVectorValueAtTimeEx(FVectorVMConte
 	}
 }
 
-template<typename PointIDParamType>
 void UNiagaraDataInterfaceHoudiniCSV::GetPointLife(FVectorVMContext& Context)
 {
-	PointIDParamType PointIDParam(Context);
+	VectorVM::FExternalFuncInputHandler<int32> PointIDParam(Context);
 
-	FRegisterHandler<float> OutValue(Context);
+	VectorVM::FExternalFuncRegisterHandler<float> OutValue(Context);
 
 	for (int32 i = 0; i < Context.NumInstances; ++i)
 	{
@@ -1285,13 +1266,13 @@ void UNiagaraDataInterfaceHoudiniCSV::GetPointLife(FVectorVMContext& Context)
 }
 
 /*
-template<typename PointIDParamType, typename TimeParamType>
+template<typename VectorVM::FExternalFuncInputHandler<int32>, typename VectorVM::FExternalFuncInputHandler<float>>
 void UNiagaraDataInterfaceHoudiniCSV::GetPointLifeAtTime(FVectorVMContext& Context)
 {
-	PointIDParamType PointIDParam(Context);
-	TimeParamType TimeParam(Context);
+	VectorVM::FExternalFuncInputHandler<int32> PointIDParam(Context);
+	VectorVM::FExternalFuncInputHandler<float> TimeParam(Context);
 
-	FRegisterHandler<float> OutValue(Context);
+	VectorVM::FExternalFuncRegisterHandler<float> OutValue(Context);
 
 	for (int32 i = 0; i < Context.NumInstances; ++i)
 	{
@@ -1314,12 +1295,11 @@ void UNiagaraDataInterfaceHoudiniCSV::GetPointLifeAtTime(FVectorVMContext& Conte
 }
 */
 
-template<typename PointIDParamType>
 void UNiagaraDataInterfaceHoudiniCSV::GetPointType(FVectorVMContext& Context)
 {
-	PointIDParamType PointIDParam(Context);
+	VectorVM::FExternalFuncInputHandler<int32> PointIDParam(Context);
 
-	FRegisterHandler<int32> OutValue(Context);
+	VectorVM::FExternalFuncRegisterHandler<int32> OutValue(Context);
 
 	for (int32 i = 0; i < Context.NumInstances; ++i)
 	{
@@ -1341,21 +1321,21 @@ void UNiagaraDataInterfaceHoudiniCSV::GetPointType(FVectorVMContext& Context)
 
 void UNiagaraDataInterfaceHoudiniCSV::GetNumberOfRows(FVectorVMContext& Context)
 {
-	FRegisterHandler<int32> OutNumRows(Context);
+	VectorVM::FExternalFuncRegisterHandler<int32> OutNumRows(Context);
 	*OutNumRows.GetDest() = HoudiniCSVAsset ? HoudiniCSVAsset->GetNumberOfRows() : 0;
 	OutNumRows.Advance();
 }
 
 void UNiagaraDataInterfaceHoudiniCSV::GetNumberOfColumns(FVectorVMContext& Context)
 {
-	FRegisterHandler<int32> OutNumCols(Context);
+	VectorVM::FExternalFuncRegisterHandler<int32> OutNumCols(Context);
 	*OutNumCols.GetDest() = HoudiniCSVAsset ? HoudiniCSVAsset->GetNumberOfColumns() : 0;
 	OutNumCols.Advance();
 }
 
 void UNiagaraDataInterfaceHoudiniCSV::GetNumberOfPoints(FVectorVMContext& Context)
 {
-	FRegisterHandler<int32> OutNumPoints(Context);
+	VectorVM::FExternalFuncRegisterHandler<int32> OutNumPoints(Context);
 	*OutNumPoints.GetDest() = HoudiniCSVAsset ? HoudiniCSVAsset->GetNumberOfPoints() : 0;
 	OutNumPoints.Advance();
 }
diff --git a/Engine/Plugins/FX/Niagara/Source/Niagara/Classes/NiagaraDataInterface.h b/Engine/Plugins/FX/Niagara/Source/Niagara/Classes/NiagaraDataInterface.h
--- a/Engine/Plugins/FX/Niagara/Source/Niagara/Classes/NiagaraDataInterface.h
+++ b/Engine/Plugins/FX/Niagara/Source/Niagara/Classes/NiagaraDataInterface.h
@@ -53,11 +53,11 @@ struct TNDIParamBinder
 	{
 		if (BindingInfo.InputParamLocations[ParamIdx])
 		{
-			NextBinder::template Bind<ParamTypes..., FConstantHandler<DataType>>(Interface, BindingInfo, InstanceData, OutFunc);
+			NextBinder::template Bind<ParamTypes..., VectorVM::FExternalFuncConstHandler<DataType>>(Interface, BindingInfo, InstanceData, OutFunc);
 		}
 		else
 		{
-			NextBinder::template Bind<ParamTypes..., FRegisterHandler<DataType>>(Interface, BindingInfo, InstanceData, OutFunc);
+			NextBinder::template Bind<ParamTypes..., VectorVM::FExternalFuncRegisterHandler<DataType>>(Interface, BindingInfo, InstanceData, OutFunc);
 		}
 	}
 };
@@ -71,7 +71,6 @@ struct TNDIParamBinder<ParamIdx, DataType, TNDINoopBinder>
 	}
 };
 
-//Helper macros allowing us to define the final binding structs for each vm external function function more concisely.
 #define NDI_FUNC_BINDER(ClassName, FuncName) T##ClassName##_##FuncName##Binder
 
 #define DEFINE_NDI_FUNC_BINDER(ClassName, FuncName)\
@@ -80,21 +79,17 @@ struct NDI_FUNC_BINDER(ClassName, FuncName)\
 	template<typename ... ParamTypes>\
 	static void Bind(UNiagaraDataInterface* Interface, const FVMExternalFunctionBindingInfo& BindingInfo, void* InstanceData, FVMExternalFunction &OutFunc)\
 	{\
-		FVMExternalFunction::CreateUObject(CastChecked<ClassName>(Interface), &ClassName::FuncName<ParamTypes...>, OutFunc);\
+		auto Lambda = [Interface](FVectorVMContext& Context) { static_cast<ClassName*>(Interface)->FuncName<ParamTypes...>(Context); };\
+		OutFunc = FVMExternalFunction::CreateLambda(Lambda);\
 	}\
 };
 
-
-
-#define NDI_RAW_FUNC_BINDER(ClassName, FuncName) T##ClassName##_##FuncName##Binder
-
-#define DEFINE_NDI_RAW_FUNC_BINDER(ClassName, FuncName)\
-struct NDI_RAW_FUNC_BINDER(ClassName, FuncName)\
+#define DEFINE_NDI_DIRECT_FUNC_BINDER(ClassName, FuncName)\
+struct NDI_FUNC_BINDER(ClassName, FuncName)\
 {\
-	template<typename ... ParamTypes>\
-	static void Bind(UNiagaraDataInterface* Interface, const FVMExternalFunctionBindingInfo& BindingInfo, void* InstanceData, FVMExternalFunction &OutFunc)\
+	static void Bind(UNiagaraDataInterface* Interface, FVMExternalFunction &OutFunc)\
 	{\
-		auto Lambda = [Interface](FVectorVMContext& Context) { static_cast<ClassName*>(Interface)->FuncName<ParamTypes...>(Context); };\
+		auto Lambda = [Interface](FVectorVMContext& Context) { static_cast<ClassName*>(Interface)->FuncName(Context); };\
 		OutFunc = FVMExternalFunction::CreateLambda(Lambda);\
 	}\
 };
diff --git a/Engine/Plugins/FX/Niagara/Source/Niagara/Classes/NiagaraDataInterfaceCollisionQuery.h b/Engine/Plugins/FX/Niagara/Source/Niagara/Classes/NiagaraDataInterfaceCollisionQuery.h
--- a/Engine/Plugins/FX/Niagara/Source/Niagara/Classes/NiagaraDataInterfaceCollisionQuery.h
+++ b/Engine/Plugins/FX/Niagara/Source/Niagara/Classes/NiagaraDataInterfaceCollisionQuery.h
@@ -56,13 +56,8 @@ public:
 	virtual void GetVMExternalFunction(const FVMExternalFunctionBindingInfo& BindingInfo, void* InstanceData, FVMExternalFunction &OutFunc) override;
 
 	// VM functions
-	template<typename PosTypeX, typename PosTypeY, typename PosTypeZ, typename VelTypeX, typename VelTypeY, typename VelTypeZ, typename DTType>
 	void SubmitQuery(FVectorVMContext& Context);
-
-	template<typename IDType>
 	void ReadQuery(FVectorVMContext& Context);
-
-	template<typename InQueryIDType, typename PosTypeX, typename PosTypeY, typename PosTypeZ, typename DirTypeX, typename DirTypeY, typename DirTypeZ, typename DTType, typename SizeType, typename DepthBoundsType>
 	void PerformQuery(FVectorVMContext& Context);
 
 	virtual bool Equals(const UNiagaraDataInterface* Other) const override;
diff --git a/Engine/Plugins/FX/Niagara/Source/Niagara/Classes/NiagaraDataInterfaceColorCurve.h b/Engine/Plugins/FX/Niagara/Source/Niagara/Classes/NiagaraDataInterfaceColorCurve.h
--- a/Engine/Plugins/FX/Niagara/Source/Niagara/Classes/NiagaraDataInterfaceColorCurve.h
+++ b/Engine/Plugins/FX/Niagara/Source/Niagara/Classes/NiagaraDataInterfaceColorCurve.h
@@ -50,7 +50,7 @@ public:
 	virtual void GetFunctions(TArray<FNiagaraFunctionSignature>& OutFunctions)override;
 	virtual void GetVMExternalFunction(const FVMExternalFunctionBindingInfo& BindingInfo, void* InstanceData, FVMExternalFunction &OutFunc) override;
 
-	template<typename UseLUT, typename XParamType>
+	template<typename UseLUT>
 	void SampleCurve(FVectorVMContext& Context);
 
 	virtual bool Equals(const UNiagaraDataInterface* Other) const override;
diff --git a/Engine/Plugins/FX/Niagara/Source/Niagara/Classes/NiagaraDataInterfaceCurlNoise.h b/Engine/Plugins/FX/Niagara/Source/Niagara/Classes/NiagaraDataInterfaceCurlNoise.h
--- a/Engine/Plugins/FX/Niagara/Source/Niagara/Classes/NiagaraDataInterfaceCurlNoise.h
+++ b/Engine/Plugins/FX/Niagara/Source/Niagara/Classes/NiagaraDataInterfaceCurlNoise.h
@@ -32,7 +32,6 @@ public:
 	virtual bool CanExecuteOnTarget(ENiagaraSimTarget Target)const override { return true; }
 	//UNiagaraDataInterface Interface End
 
-	template<typename XType, typename YType, typename ZType>
 	void SampleNoiseField(FVectorVMContext& Context);
 
 	virtual bool Equals(const UNiagaraDataInterface* Other) const override;
diff --git a/Engine/Plugins/FX/Niagara/Source/Niagara/Classes/NiagaraDataInterfaceCurve.h b/Engine/Plugins/FX/Niagara/Source/Niagara/Classes/NiagaraDataInterfaceCurve.h
--- a/Engine/Plugins/FX/Niagara/Source/Niagara/Classes/NiagaraDataInterfaceCurve.h
+++ b/Engine/Plugins/FX/Niagara/Source/Niagara/Classes/NiagaraDataInterfaceCurve.h
@@ -36,9 +36,9 @@ public:
 	virtual void GetFunctions(TArray<FNiagaraFunctionSignature>& OutFunctions)override;
 	virtual void GetVMExternalFunction(const FVMExternalFunctionBindingInfo& BindingInfo, void* InstanceData, FVMExternalFunction &OutFunc) override;
 
-	template<typename UseLUT, typename XParamType>
+	template<typename UseLUT>
 	void SampleCurve(FVectorVMContext& Context);
-
+	
 	virtual bool Equals(const UNiagaraDataInterface* Other) const override;
 
 	//~ UNiagaraDataInterfaceCurveBase interface
diff --git a/Engine/Plugins/FX/Niagara/Source/Niagara/Classes/NiagaraDataInterfaceSkeletalMesh.h b/Engine/Plugins/FX/Niagara/Source/Niagara/Classes/NiagaraDataInterfaceSkeletalMesh.h
--- a/Engine/Plugins/FX/Niagara/Source/Niagara/Classes/NiagaraDataInterfaceSkeletalMesh.h
+++ b/Engine/Plugins/FX/Niagara/Source/Niagara/Classes/NiagaraDataInterfaceSkeletalMesh.h
@@ -365,25 +365,27 @@ public:
 	template<typename FilterMode, typename AreaWeightingMode>
 	void GetFilteredTriangleCount(FVectorVMContext& Context);
 
-	template<typename FilterMode, typename AreaWeightingMode, typename TriType>
+	template<typename FilterMode, typename AreaWeightingMode>
 	void GetFilteredTriangleAt(FVectorVMContext& Context);
 
 	template<typename FilterMode, typename AreaWeightingMode>
 	void RandomTriCoord(FVectorVMContext& Context);
 
-	template<typename FilterMode, typename AreaWeightingMode, typename TriType, typename BaryXType, typename BaryYType, typename BaryZType>
+	template<typename FilterMode, typename AreaWeightingMode>
 	void IsValidTriCoord(FVectorVMContext& Context);
 
-	template<typename SkinningHandlerType, typename TransformHandlerType, typename VertexAccessorType, typename TriType, typename BaryXType, typename BaryYType, typename BaryZType>
+	template<typename SkinningHandlerType, typename TransformHandlerType, typename VertexAccessorType>
+	void GetTriCoordSkinnedData_Opt(FVectorVMContext& Context);
+
+	template<typename SkinningHandlerType, typename TransformHandlerType, typename VertexAccessorType>
 	void GetTriCoordSkinnedData(FVectorVMContext& Context);
 
-	template<typename TriType, typename BaryXType, typename BaryYType, typename BaryZType>
 	void GetTriCoordColor(FVectorVMContext& Context);
 
-	template<typename VertexAccessorType, typename TriType, typename BaryXType, typename BaryYType, typename BaryZType, typename UVSetType>
+	template<typename VertexAccessorType>
 	void GetTriCoordUV(FVectorVMContext& Context);
 
-	template<typename SkinningHandlerType, typename TriType>
+	template<typename SkinningHandlerType>
 	void GetTriCoordVertices(FVectorVMContext& Context);
 		
 private:
@@ -409,22 +411,21 @@ public:
 	template<typename FilterMode>
 	void GetFilteredVertexCount(FVectorVMContext& Context);
 
-	template<typename FilterMode, typename VertexType>
+	template<typename FilterMode>
 	void GetFilteredVertexAt(FVectorVMContext& Context);
 
 	template<typename FilterMode>
 	void RandomVertex(FVectorVMContext& Context);
 
-	template<typename FilterMode, typename VertexType>
+	template<typename FilterMode>
 	void IsValidVertex(FVectorVMContext& Context);
 
-	template<typename SkinningHandlerType, typename TransformHandlerType, typename VertexAccessorType, typename VertexType>
+	template<typename SkinningHandlerType, typename TransformHandlerType, typename VertexAccessorType>
 	void GetVertexSkinnedData(FVectorVMContext& Context);
 
-	template<typename VertexType>
 	void GetVertexColor(FVectorVMContext& Context);
 
-	template<typename VertexAccessorType, typename VertexType, typename UVSetType>
+	template<typename VertexAccessorType>
 	void GetVertexUV(FVectorVMContext& Context);
 
 private:
@@ -447,22 +448,19 @@ public:
 	void GetSkeletonSamplingFunctions(TArray<FNiagaraFunctionSignature>& OutFunctions);
 	void BindSkeletonSamplingFunction(const FVMExternalFunctionBindingInfo& BindingInfo, FNDISkeletalMesh_InstanceData* InstData, FVMExternalFunction &OutFunc);
 
-	template<typename BoneType>
 	void IsValidBone(FVectorVMContext& Context);
 
 	void GetSpecificBoneCount(FVectorVMContext& Context);
 
-	template<typename BoneType>
 	void GetSpecificBoneAt(FVectorVMContext& Context);
 
 	void RandomSpecificBone(FVectorVMContext& Context);
 
-	template<typename SkinningHandlerType, typename TransformHandlerType, typename BoneType>
+	template<typename SkinningHandlerType, typename TransformHandlerType>
 	void GetSkinnedBoneData(FVectorVMContext& Context);
 	
 	void GetSpecificSocketCount(FVectorVMContext& Context);
 
-	template<typename SocketType>
 	void GetSpecificSocketBoneAt(FVectorVMContext& Context);
 
 	void RandomSpecificSocketBone(FVectorVMContext& Context);
diff --git a/Engine/Plugins/FX/Niagara/Source/Niagara/Classes/NiagaraDataInterfaceStaticMesh.h b/Engine/Plugins/FX/Niagara/Source/Niagara/Classes/NiagaraDataInterfaceStaticMesh.h
--- a/Engine/Plugins/FX/Niagara/Source/Niagara/Classes/NiagaraDataInterfaceStaticMesh.h
+++ b/Engine/Plugins/FX/Niagara/Source/Niagara/Classes/NiagaraDataInterfaceStaticMesh.h
@@ -148,35 +148,32 @@ public:
 	template<typename TAreaWeighted>
 	void RandomTriCoord(FVectorVMContext& Context);
 
-	template<typename TAreaWeighted, typename SectionIdxType>
+	template<typename TAreaWeighted>
 	void RandomTriCoordOnSection(FVectorVMContext& Context);
 
- 	template<typename InputType0, typename InputType1>
  	void RandomTriCoordVertexColorFiltered(FVectorVMContext& Context);
 	
-	template<typename TransformHandlerType, typename TriType, typename BaryXType, typename BaryYType, typename BaryZType>
+	template<typename TransformHandlerType>
 	void GetTriCoordPosition(FVectorVMContext& Context);
 
-	template<typename TransformHandlerType, typename TriType, typename BaryXType, typename BaryYType, typename BaryZType>
+	template<typename TransformHandlerType>
 	void GetTriCoordNormal(FVectorVMContext& Context);
 
-	template<typename VertexAccessorType, typename TransformHandlerType, typename TriType, typename BaryXType, typename BaryYType, typename BaryZType>
+	template<typename VertexAccessorType, typename TransformHandlerType>
 	void GetTriCoordTangents(FVectorVMContext& Context);
 
-	template<typename TriType, typename BaryXType, typename BaryYType, typename BaryZType>
 	void GetTriCoordColor(FVectorVMContext& Context);
 
-	template<typename VertexAccessorType, typename TriType, typename BaryXType, typename BaryYType, typename BaryZType, typename UVSetType>
+	template<typename VertexAccessorType>
 	void GetTriCoordUV(FVectorVMContext& Context);
 
-	template<typename TriType, typename BaryXType, typename BaryYType, typename BaryZType>
 	void GetTriCoordPositionAndVelocity(FVectorVMContext& Context);
 
 	void GetLocalToWorld(FVectorVMContext& Context);
 	void GetLocalToWorldInverseTransposed(FVectorVMContext& Context);
 	void GetWorldVelocity(FVectorVMContext& Context);
 
-	template<typename TransformHandlerType, typename VertexIndexType>
+	template<typename TransformHandlerType>
 	void GetVertexPosition(FVectorVMContext& Context);
 
 	FORCEINLINE_DEBUGGABLE bool UsesSectionFilter()const { return SectionFilter.CanEverReject(); }
diff --git a/Engine/Plugins/FX/Niagara/Source/Niagara/Classes/NiagaraDataInterfaceTexture.h b/Engine/Plugins/FX/Niagara/Source/Niagara/Classes/NiagaraDataInterfaceTexture.h
--- a/Engine/Plugins/FX/Niagara/Source/Niagara/Classes/NiagaraDataInterfaceTexture.h
+++ b/Engine/Plugins/FX/Niagara/Source/Niagara/Classes/NiagaraDataInterfaceTexture.h
@@ -33,7 +33,6 @@ public:
 	virtual bool CanExecuteOnTarget(ENiagaraSimTarget Target)const override { return Target==ENiagaraSimTarget::GPUComputeSim; }
 	//UNiagaraDataInterface Interface End
 
-	template<typename XType, typename YType>
 	void SampleTexture(FVectorVMContext& Context);
 	void GetTextureDimensions(FVectorVMContext& Context);
 	void NoOp(FVectorVMContext& Context);
diff --git a/Engine/Plugins/FX/Niagara/Source/Niagara/Classes/NiagaraDataInterfaceVector2DCurve.h b/Engine/Plugins/FX/Niagara/Source/Niagara/Classes/NiagaraDataInterfaceVector2DCurve.h
--- a/Engine/Plugins/FX/Niagara/Source/Niagara/Classes/NiagaraDataInterfaceVector2DCurve.h
+++ b/Engine/Plugins/FX/Niagara/Source/Niagara/Classes/NiagaraDataInterfaceVector2DCurve.h
@@ -45,7 +45,7 @@ public:
 
 	virtual void GetVMExternalFunction(const FVMExternalFunctionBindingInfo& BindingInfo, void* InstanceData, FVMExternalFunction &OutFunc) override;
 
-	template<typename UseLUT, typename XParamType>
+	template<typename UseLUT>
 	void SampleCurve(FVectorVMContext& Context);
 
 	//~ UNiagaraDataInterfaceCurveBase interface
diff --git a/Engine/Plugins/FX/Niagara/Source/Niagara/Classes/NiagaraDataInterfaceVector4Curve.h b/Engine/Plugins/FX/Niagara/Source/Niagara/Classes/NiagaraDataInterfaceVector4Curve.h
--- a/Engine/Plugins/FX/Niagara/Source/Niagara/Classes/NiagaraDataInterfaceVector4Curve.h
+++ b/Engine/Plugins/FX/Niagara/Source/Niagara/Classes/NiagaraDataInterfaceVector4Curve.h
@@ -50,7 +50,7 @@ public:
 	virtual void GetFunctions(TArray<FNiagaraFunctionSignature>& OutFunctions)override;
 	virtual void GetVMExternalFunction(const FVMExternalFunctionBindingInfo& BindingInfo, void* InstanceData, FVMExternalFunction &OutFunc) override;
 
-	template<typename UseLUT, typename XParamType>
+	template<typename UseLUT>
 	void SampleCurve(FVectorVMContext& Context);
 
 	virtual bool Equals(const UNiagaraDataInterface* Other) const override;
diff --git a/Engine/Plugins/FX/Niagara/Source/Niagara/Classes/NiagaraDataInterfaceVectorCurve.h b/Engine/Plugins/FX/Niagara/Source/Niagara/Classes/NiagaraDataInterfaceVectorCurve.h
--- a/Engine/Plugins/FX/Niagara/Source/Niagara/Classes/NiagaraDataInterfaceVectorCurve.h
+++ b/Engine/Plugins/FX/Niagara/Source/Niagara/Classes/NiagaraDataInterfaceVectorCurve.h
@@ -47,7 +47,7 @@ public:
 	virtual void GetFunctions(TArray<FNiagaraFunctionSignature>& OutFunctions)override;
 	virtual void GetVMExternalFunction(const FVMExternalFunctionBindingInfo& BindingInfo, void* InstanceData, FVMExternalFunction &OutFunc) override;
 
-	template<typename UseLUT, typename XParamType>
+	template<typename UseLUT>
 	void SampleCurve(FVectorVMContext& Context);
 
 	virtual bool Equals(const UNiagaraDataInterface* Other) const override;
diff --git a/Engine/Plugins/FX/Niagara/Source/Niagara/Classes/NiagaraDataInterfaceVectorField.h b/Engine/Plugins/FX/Niagara/Source/Niagara/Classes/NiagaraDataInterfaceVectorField.h
--- a/Engine/Plugins/FX/Niagara/Source/Niagara/Classes/NiagaraDataInterfaceVectorField.h
+++ b/Engine/Plugins/FX/Niagara/Source/Niagara/Classes/NiagaraDataInterfaceVectorField.h
@@ -72,7 +72,6 @@ public:
 	virtual FNiagaraDataInterfaceParametersCS* ConstructComputeParameters()const override;
 	//~ UNiagaraDataInterface interface END
 
-	template<typename XType, typename YType, typename ZType>
 	void SampleVectorField(FVectorVMContext& Context);	
 
 	void GetFieldDimensions(FVectorVMContext& Context);
diff --git a/Engine/Plugins/FX/Niagara/Source/Niagara/Private/NDISkeletalMesh_BoneSampling.cpp b/Engine/Plugins/FX/Niagara/Source/Niagara/Private/NDISkeletalMesh_BoneSampling.cpp
--- a/Engine/Plugins/FX/Niagara/Source/Niagara/Private/NDISkeletalMesh_BoneSampling.cpp
+++ b/Engine/Plugins/FX/Niagara/Source/Niagara/Private/NDISkeletalMesh_BoneSampling.cpp
@@ -14,11 +14,11 @@
 DECLARE_CYCLE_STAT(TEXT("Skel Mesh Skeleton Sampling"), STAT_NiagaraSkel_Bone_Sample, STATGROUP_Niagara);
 
 //Final binders for all static mesh interface functions.
-DEFINE_NDI_RAW_FUNC_BINDER(UNiagaraDataInterfaceSkeletalMesh, GetSkinnedBoneData)
-DEFINE_NDI_RAW_FUNC_BINDER(UNiagaraDataInterfaceSkeletalMesh, IsValidBone)
-DEFINE_NDI_RAW_FUNC_BINDER(UNiagaraDataInterfaceSkeletalMesh, GetSpecificBoneAt)
+DEFINE_NDI_FUNC_BINDER(UNiagaraDataInterfaceSkeletalMesh, GetSkinnedBoneData)
+DEFINE_NDI_DIRECT_FUNC_BINDER(UNiagaraDataInterfaceSkeletalMesh, IsValidBone)
+DEFINE_NDI_DIRECT_FUNC_BINDER(UNiagaraDataInterfaceSkeletalMesh, GetSpecificBoneAt)
 
-DEFINE_NDI_RAW_FUNC_BINDER(UNiagaraDataInterfaceSkeletalMesh, GetSpecificSocketBoneAt)
+DEFINE_NDI_DIRECT_FUNC_BINDER(UNiagaraDataInterfaceSkeletalMesh, GetSpecificSocketBoneAt)
 
 static const FName RandomSpecificBoneName("RandomSpecificBone");
 static const FName IsValidBoneName("IsValidBoneName");
@@ -173,17 +173,17 @@ void UNiagaraDataInterfaceSkeletalMesh::BindSkeletonSamplingFunction(const FVMEx
 	else if (BindingInfo.Name == IsValidBoneName)
 	{
 		check(BindingInfo.GetNumInputs() == 2 && BindingInfo.GetNumOutputs() == 1);
-		TNDIParamBinder<0, int32, NDI_RAW_FUNC_BINDER(UNiagaraDataInterfaceSkeletalMesh, IsValidBone)>::Bind(this, BindingInfo, InstanceData, OutFunc);
+		NDI_FUNC_BINDER(UNiagaraDataInterfaceSkeletalMesh, IsValidBone)::Bind(this, OutFunc);
 	}
 	else if (BindingInfo.Name == GetSkinnedBoneDataName)
 	{
 		check(BindingInfo.GetNumInputs() == 2 && BindingInfo.GetNumOutputs() == 6);
-		TSkinningModeBinder<TNDIExplicitBinder<FNDITransformHandlerNoop, TNDIParamBinder<0, int32, NDI_RAW_FUNC_BINDER(UNiagaraDataInterfaceSkeletalMesh, GetSkinnedBoneData)>>>::Bind(this, BindingInfo, InstanceData, OutFunc);
+		TSkinningModeBinder<TNDIExplicitBinder<FNDITransformHandlerNoop, NDI_FUNC_BINDER(UNiagaraDataInterfaceSkeletalMesh, GetSkinnedBoneData)>>::Bind(this, BindingInfo, InstanceData, OutFunc);
 	}
 	else if (BindingInfo.Name == GetSkinnedBoneDataWSName)
 	{
 		check(BindingInfo.GetNumInputs() == 2 && BindingInfo.GetNumOutputs() == 6);
-		TSkinningModeBinder<TNDIExplicitBinder<FNDITransformHandler, TNDIParamBinder<0, int32, NDI_RAW_FUNC_BINDER(UNiagaraDataInterfaceSkeletalMesh, GetSkinnedBoneData)>>>::Bind(this, BindingInfo, InstanceData, OutFunc);
+		TSkinningModeBinder<TNDIExplicitBinder<FNDITransformHandler, NDI_FUNC_BINDER(UNiagaraDataInterfaceSkeletalMesh, GetSkinnedBoneData)>>::Bind(this, BindingInfo, InstanceData, OutFunc);
 	}
 	else if (BindingInfo.Name == GetSpecificBoneCountName)
 	{
@@ -194,7 +194,7 @@ void UNiagaraDataInterfaceSkeletalMesh::BindSkeletonSamplingFunction(const FVMEx
 	else if (BindingInfo.Name == GetSpecificBoneAtName)
 	{
 		check(BindingInfo.GetNumInputs() == 2 && BindingInfo.GetNumOutputs() == 1);
-		TNDIParamBinder<0, int32, NDI_RAW_FUNC_BINDER(UNiagaraDataInterfaceSkeletalMesh, GetSpecificBoneAt)>::Bind(this, BindingInfo, InstanceData, OutFunc);
+		NDI_FUNC_BINDER(UNiagaraDataInterfaceSkeletalMesh, GetSpecificBoneAt)::Bind(this, OutFunc);
 	}
 	//Socket Functions
 	else if (BindingInfo.Name == RandomSpecificSocketBoneName)
@@ -212,7 +212,7 @@ void UNiagaraDataInterfaceSkeletalMesh::BindSkeletonSamplingFunction(const FVMEx
 	else if (BindingInfo.Name == GetSpecificSocketBoneAtName)
 	{
 		check(BindingInfo.GetNumInputs() == 2 && BindingInfo.GetNumOutputs() == 1);
-		TNDIParamBinder<0, int32, NDI_RAW_FUNC_BINDER(UNiagaraDataInterfaceSkeletalMesh, GetSpecificSocketBoneAt)>::Bind(this, BindingInfo, InstanceData, OutFunc);
+		NDI_FUNC_BINDER(UNiagaraDataInterfaceSkeletalMesh, GetSpecificSocketBoneAt)::Bind(this, OutFunc);
 	}
 }
 
@@ -223,9 +223,9 @@ void UNiagaraDataInterfaceSkeletalMesh::BindSkeletonSamplingFunction(const FVMEx
 void UNiagaraDataInterfaceSkeletalMesh::GetSpecificBoneCount(FVectorVMContext& Context)
 {
 	SCOPE_CYCLE_COUNTER(STAT_NiagaraSkel_Bone_Sample);
-	FUserPtrHandler<FNDISkeletalMesh_InstanceData> InstData(Context);
+	VectorVM::FUserPtrHandler<FNDISkeletalMesh_InstanceData> InstData(Context);
 
-	FRegisterHandler<int32> OutCount(Context);
+	VectorVM::FExternalFuncRegisterHandler<int32> OutCount(Context);
 
 	int32 Num = SpecificBones.Num();
 	for (int32 i = 0; i < Context.NumInstances; ++i)
@@ -234,15 +234,14 @@ void UNiagaraDataInterfaceSkeletalMesh::GetSpecificBoneCount(FVectorVMContext& C
 	}
 }
 
-template<typename BoneType>
 void UNiagaraDataInterfaceSkeletalMesh::GetSpecificBoneAt(FVectorVMContext& Context)
 {
 	SCOPE_CYCLE_COUNTER(STAT_NiagaraSkel_Bone_Sample);
 
-	BoneType BoneParam(Context);
-	FUserPtrHandler<FNDISkeletalMesh_InstanceData> InstData(Context);
+	VectorVM::FExternalFuncInputHandler<int32> BoneParam(Context);
+	VectorVM::FUserPtrHandler<FNDISkeletalMesh_InstanceData> InstData(Context);
 
-	FRegisterHandler<int32> OutBone(Context);
+	VectorVM::FExternalFuncRegisterHandler<int32> OutBone(Context);
 	const TArray<int32>& SpecificBonesArray = InstData->SpecificBones;
 
 	int32 Max = SpecificBones.Num() - 1;
@@ -264,9 +263,9 @@ void UNiagaraDataInterfaceSkeletalMesh::RandomSpecificBone(FVectorVMContext& Con
 {
 	SCOPE_CYCLE_COUNTER(STAT_NiagaraSkel_Bone_Sample);
 
-	FUserPtrHandler<FNDISkeletalMesh_InstanceData> InstData(Context);
+	VectorVM::FUserPtrHandler<FNDISkeletalMesh_InstanceData> InstData(Context);
 
-	FRegisterHandler<int32> OutBone(Context);
+	VectorVM::FExternalFuncRegisterHandler<int32> OutBone(Context);
 	const TArray<int32>& SpecificBonesArray = InstData->SpecificBones;
 
 	int32 Max = SpecificBones.Num() - 1;
@@ -284,18 +283,17 @@ void UNiagaraDataInterfaceSkeletalMesh::RandomSpecificBone(FVectorVMContext& Con
 	}
 }
 
-template<typename BoneType>
 void UNiagaraDataInterfaceSkeletalMesh::IsValidBone(FVectorVMContext& Context)
 {
 	SCOPE_CYCLE_COUNTER(STAT_NiagaraSkel_Bone_Sample);
 
-	BoneType BoneParam(Context);
+	VectorVM::FExternalFuncInputHandler<int32> BoneParam(Context);
 
-	FUserPtrHandler<FNDISkeletalMesh_InstanceData> InstData(Context);
+	VectorVM::FUserPtrHandler<FNDISkeletalMesh_InstanceData> InstData(Context);
 	checkfSlow(InstData.Get(), TEXT("Skeletal Mesh Interface has invalid instance data. %s"), *GetPathName());
 	checkfSlow(InstData->Mesh, TEXT("Skeletal Mesh Interface has invalid mesh. %s"), *GetPathName());
 
-	FRegisterHandler<FNiagaraBool> OutValid(Context);
+	VectorVM::FExternalFuncRegisterHandler<FNiagaraBool> OutValid(Context);
 
 	FSkeletalMeshAccessorHelper MeshAccessor;
 	MeshAccessor.Init<TIntegralConstant<int32, 0>, TIntegralConstant<int32, 0>>(InstData);
@@ -321,8 +319,8 @@ struct FBoneSocketSkinnedDataOutputHandler
 	{
 	}
 
-	FRegisterHandler<float> PosX; FRegisterHandler<float> PosY; FRegisterHandler<float> PosZ;
-	FRegisterHandler<float> VelX; FRegisterHandler<float> VelY; FRegisterHandler<float> VelZ;
+	VectorVM::FExternalFuncRegisterHandler<float> PosX; VectorVM::FExternalFuncRegisterHandler<float> PosY; VectorVM::FExternalFuncRegisterHandler<float> PosZ;
+	VectorVM::FExternalFuncRegisterHandler<float> VelX; VectorVM::FExternalFuncRegisterHandler<float> VelY; VectorVM::FExternalFuncRegisterHandler<float> VelZ;
 
 	//TODO: Rotation + Scale too? Use quats so we can get proper interpolation between bone and parent.
 
@@ -344,15 +342,15 @@ struct FBoneSocketSkinnedDataOutputHandler
 	}
 };
 
-template<typename SkinningHandlerType, typename TransformHandlerType, typename BoneType>
+template<typename SkinningHandlerType, typename TransformHandlerType>
 void UNiagaraDataInterfaceSkeletalMesh::GetSkinnedBoneData(FVectorVMContext& Context)
 {
 	SCOPE_CYCLE_COUNTER(STAT_NiagaraSkel_Bone_Sample);
 	SkinningHandlerType SkinningHandler;
 	TransformHandlerType TransformHandler;
-	BoneType BoneParam(Context);
+	VectorVM::FExternalFuncInputHandler<int32> BoneParam(Context);
 
-	FUserPtrHandler<FNDISkeletalMesh_InstanceData> InstData(Context);
+	VectorVM::FUserPtrHandler<FNDISkeletalMesh_InstanceData> InstData(Context);
 
 	checkfSlow(InstData.Get(), TEXT("Skeletal Mesh Interface has invalid instance data. %s"), *GetPathName());
 	checkfSlow(InstData->Mesh, TEXT("Skeletal Mesh Interface has invalid mesh. %s"), *GetPathName());
@@ -410,9 +408,9 @@ void UNiagaraDataInterfaceSkeletalMesh::GetSkinnedBoneData(FVectorVMContext& Con
 void UNiagaraDataInterfaceSkeletalMesh::GetSpecificSocketCount(FVectorVMContext& Context)
 {
 	SCOPE_CYCLE_COUNTER(STAT_NiagaraSkel_Bone_Sample);
-	FUserPtrHandler<FNDISkeletalMesh_InstanceData> InstData(Context);
+	VectorVM::FUserPtrHandler<FNDISkeletalMesh_InstanceData> InstData(Context);
 
-	FRegisterHandler<int32> OutCount(Context);
+	VectorVM::FExternalFuncRegisterHandler<int32> OutCount(Context);
 
 	int32 Num = SpecificSockets.Num();
 	for (int32 i = 0; i < Context.NumInstances; ++i)
@@ -421,15 +419,14 @@ void UNiagaraDataInterfaceSkeletalMesh::GetSpecificSocketCount(FVectorVMContext&
 	}
 }
 
-template<typename SocketType>
 void UNiagaraDataInterfaceSkeletalMesh::GetSpecificSocketBoneAt(FVectorVMContext& Context)
 {
 	SCOPE_CYCLE_COUNTER(STAT_NiagaraSkel_Bone_Sample);
 
-	SocketType SocketParam(Context);
-	FUserPtrHandler<FNDISkeletalMesh_InstanceData> InstData(Context);
+	VectorVM::FExternalFuncInputHandler<int32> SocketParam(Context);
+	VectorVM::FUserPtrHandler<FNDISkeletalMesh_InstanceData> InstData(Context);
 
-	FRegisterHandler<int32> OutSocketBone(Context);
+	VectorVM::FExternalFuncRegisterHandler<int32> OutSocketBone(Context);
 	const TArray<int32>& SpecificSocketsArray = InstData->SpecificSocketBones;
 
 	int32 Max = SpecificSockets.Num() - 1;
@@ -451,9 +448,9 @@ void UNiagaraDataInterfaceSkeletalMesh::RandomSpecificSocketBone(FVectorVMContex
 {
 	SCOPE_CYCLE_COUNTER(STAT_NiagaraSkel_Bone_Sample);
 
-	FUserPtrHandler<FNDISkeletalMesh_InstanceData> InstData(Context);
+	VectorVM::FUserPtrHandler<FNDISkeletalMesh_InstanceData> InstData(Context);
 
-	FRegisterHandler<int32> OutSocketBone(Context);
+	VectorVM::FExternalFuncRegisterHandler<int32> OutSocketBone(Context);
 	const TArray<int32>& SpecificSocketsArray = InstData->SpecificSocketBones;
 
 	int32 Max = SpecificSockets.Num() - 1;
diff --git a/Engine/Plugins/FX/Niagara/Source/Niagara/Private/NDISkeletalMesh_TriangleSampling.cpp b/Engine/Plugins/FX/Niagara/Source/Niagara/Private/NDISkeletalMesh_TriangleSampling.cpp
--- a/Engine/Plugins/FX/Niagara/Source/Niagara/Private/NDISkeletalMesh_TriangleSampling.cpp
+++ b/Engine/Plugins/FX/Niagara/Source/Niagara/Private/NDISkeletalMesh_TriangleSampling.cpp
@@ -14,14 +14,14 @@
 DECLARE_CYCLE_STAT(TEXT("Skel Mesh Sampling"), STAT_NiagaraSkel_Sample, STATGROUP_Niagara);
 
 //Final binders for all static mesh interface functions.
-DEFINE_NDI_RAW_FUNC_BINDER(UNiagaraDataInterfaceSkeletalMesh, RandomTriCoord);
-DEFINE_NDI_RAW_FUNC_BINDER(UNiagaraDataInterfaceSkeletalMesh, GetTriCoordSkinnedData);
-DEFINE_NDI_RAW_FUNC_BINDER(UNiagaraDataInterfaceSkeletalMesh, GetTriCoordColor);
-DEFINE_NDI_RAW_FUNC_BINDER(UNiagaraDataInterfaceSkeletalMesh, GetTriCoordUV);
-DEFINE_NDI_RAW_FUNC_BINDER(UNiagaraDataInterfaceSkeletalMesh, IsValidTriCoord);
-DEFINE_NDI_RAW_FUNC_BINDER(UNiagaraDataInterfaceSkeletalMesh, GetFilteredTriangleCount);
-DEFINE_NDI_RAW_FUNC_BINDER(UNiagaraDataInterfaceSkeletalMesh, GetFilteredTriangleAt);
-DEFINE_NDI_RAW_FUNC_BINDER(UNiagaraDataInterfaceSkeletalMesh, GetTriCoordVertices);
+DEFINE_NDI_FUNC_BINDER(UNiagaraDataInterfaceSkeletalMesh, RandomTriCoord);
+DEFINE_NDI_FUNC_BINDER(UNiagaraDataInterfaceSkeletalMesh, GetTriCoordSkinnedData);
+DEFINE_NDI_DIRECT_FUNC_BINDER(UNiagaraDataInterfaceSkeletalMesh, GetTriCoordColor);
+DEFINE_NDI_FUNC_BINDER(UNiagaraDataInterfaceSkeletalMesh, GetTriCoordUV);
+DEFINE_NDI_FUNC_BINDER(UNiagaraDataInterfaceSkeletalMesh, IsValidTriCoord);
+DEFINE_NDI_FUNC_BINDER(UNiagaraDataInterfaceSkeletalMesh, GetFilteredTriangleCount);
+DEFINE_NDI_FUNC_BINDER(UNiagaraDataInterfaceSkeletalMesh, GetFilteredTriangleAt);
+DEFINE_NDI_FUNC_BINDER(UNiagaraDataInterfaceSkeletalMesh, GetTriCoordVertices);
 
 static const FName RandomTriCoordName("RandomTriCoord");
 static const FName IsValidTriCoordName("IsValidTriCoord");
@@ -165,48 +165,48 @@ void UNiagaraDataInterfaceSkeletalMesh::BindTriangleSamplingFunction(const FVMEx
 	if (BindingInfo.Name == RandomTriCoordName)
 	{
 		check(BindingInfo.GetNumInputs() == 1 && BindingInfo.GetNumOutputs() == 4);
-		TFilterModeBinder<TAreaWeightingModeBinder<NDI_RAW_FUNC_BINDER(UNiagaraDataInterfaceSkeletalMesh, RandomTriCoord)>>::Bind(this, BindingInfo, InstanceData, OutFunc);
+		TFilterModeBinder<TAreaWeightingModeBinder<NDI_FUNC_BINDER(UNiagaraDataInterfaceSkeletalMesh, RandomTriCoord)>>::Bind(this, BindingInfo, InstanceData, OutFunc);
 	}
 	else if (BindingInfo.Name == IsValidTriCoordName)
 	{
 		check(BindingInfo.GetNumInputs() == 5 && BindingInfo.GetNumOutputs() == 1);
-		TFilterModeBinder<TAreaWeightingModeBinder<TNDIParamBinder<0, int32, TNDIParamBinder<1, float, TNDIParamBinder<2, float, TNDIParamBinder<3, float, NDI_RAW_FUNC_BINDER(UNiagaraDataInterfaceSkeletalMesh, IsValidTriCoord)>>>>>>::Bind(this, BindingInfo, InstanceData, OutFunc);
+		TFilterModeBinder<TAreaWeightingModeBinder<NDI_FUNC_BINDER(UNiagaraDataInterfaceSkeletalMesh, IsValidTriCoord)>>::Bind(this, BindingInfo, InstanceData, OutFunc);
 	}
 	else if (BindingInfo.Name == GetSkinnedTriangleDataName)
 	{
 		check(BindingInfo.GetNumInputs() == 5 && BindingInfo.GetNumOutputs() == 15);
-		TSkinningModeBinder<TNDIExplicitBinder<FNDITransformHandlerNoop, TVertexAccessorBinder<TNDIParamBinder<0, int32, TNDIParamBinder<1, float, TNDIParamBinder<2, float, TNDIParamBinder<3, float, NDI_RAW_FUNC_BINDER(UNiagaraDataInterfaceSkeletalMesh, GetTriCoordSkinnedData)>>>>>>>::Bind(this, BindingInfo, InstanceData, OutFunc);
+		TSkinningModeBinder<TNDIExplicitBinder<FNDITransformHandlerNoop, TVertexAccessorBinder<NDI_FUNC_BINDER(UNiagaraDataInterfaceSkeletalMesh, GetTriCoordSkinnedData)>>>::Bind(this, BindingInfo, InstanceData, OutFunc);
 	}
 	else if (BindingInfo.Name == GetSkinnedTriangleDataWSName)
 	{
 		check(BindingInfo.GetNumInputs() == 5 && BindingInfo.GetNumOutputs() == 15);
-		TSkinningModeBinder<TNDIExplicitBinder<FNDITransformHandler, TVertexAccessorBinder<TNDIParamBinder<0, int32, TNDIParamBinder<1, float, TNDIParamBinder<2, float, TNDIParamBinder<3, float, NDI_RAW_FUNC_BINDER(UNiagaraDataInterfaceSkeletalMesh, GetTriCoordSkinnedData)>>>>>>>::Bind(this, BindingInfo, InstanceData, OutFunc);
-	}
+		TSkinningModeBinder<TNDIExplicitBinder<FNDITransformHandler, TVertexAccessorBinder<NDI_FUNC_BINDER(UNiagaraDataInterfaceSkeletalMesh, GetTriCoordSkinnedData)>>>::Bind(this, BindingInfo, InstanceData, OutFunc);
+ 	}
 	else if (BindingInfo.Name == GetTriColorName)
 	{
 		check(BindingInfo.GetNumInputs() == 5 && BindingInfo.GetNumOutputs() == 4);
 		bNeedsVertexColors = true;
-		TNDIParamBinder<0, int32, TNDIParamBinder<1, float, TNDIParamBinder<2, float, TNDIParamBinder<3, float, NDI_RAW_FUNC_BINDER(UNiagaraDataInterfaceSkeletalMesh, GetTriCoordColor)>>>>::Bind(this, BindingInfo, InstanceData, OutFunc);
+		NDI_FUNC_BINDER(UNiagaraDataInterfaceSkeletalMesh, GetTriCoordColor)::Bind(this, OutFunc);
 	}
 	else if (BindingInfo.Name == GetTriUVName)
 	{
 		check(BindingInfo.GetNumInputs() == 6 && BindingInfo.GetNumOutputs() == 2);
-		TVertexAccessorBinder<TNDIParamBinder<0, int32, TNDIParamBinder<1, float, TNDIParamBinder<2, float, TNDIParamBinder<3, float, TNDIParamBinder<4, int32, NDI_RAW_FUNC_BINDER(UNiagaraDataInterfaceSkeletalMesh, GetTriCoordUV)>>>>>>::Bind(this, BindingInfo, InstanceData, OutFunc);
+		TVertexAccessorBinder<NDI_FUNC_BINDER(UNiagaraDataInterfaceSkeletalMesh, GetTriCoordUV)>::Bind(this, BindingInfo, InstanceData, OutFunc);
 	}
 	else if (BindingInfo.Name == GetTriangleCountName)
 	{
 		check(BindingInfo.GetNumInputs() == 1 && BindingInfo.GetNumOutputs() == 1);
-		TFilterModeBinder<TAreaWeightingModeBinder<NDI_RAW_FUNC_BINDER(UNiagaraDataInterfaceSkeletalMesh, GetFilteredTriangleCount)>>::Bind(this, BindingInfo, InstanceData, OutFunc);
+		TFilterModeBinder<TAreaWeightingModeBinder<NDI_FUNC_BINDER(UNiagaraDataInterfaceSkeletalMesh, GetFilteredTriangleCount)>>::Bind(this, BindingInfo, InstanceData, OutFunc);
 	}
 	else if (BindingInfo.Name == GetTriangleAtName)
 	{
 		check(BindingInfo.GetNumInputs() == 2 && BindingInfo.GetNumOutputs() == 4);
-		TFilterModeBinder<TAreaWeightingModeBinder<TNDIParamBinder<0, int32, NDI_RAW_FUNC_BINDER(UNiagaraDataInterfaceSkeletalMesh, GetFilteredTriangleAt)>>>::Bind(this, BindingInfo, InstanceData, OutFunc);
+		TFilterModeBinder<TAreaWeightingModeBinder<NDI_FUNC_BINDER(UNiagaraDataInterfaceSkeletalMesh, GetFilteredTriangleAt)>>::Bind(this, BindingInfo, InstanceData, OutFunc);
 	}
 	else if (BindingInfo.Name == GetTriCoordVerticesName)
 	{
 		check(BindingInfo.GetNumInputs() == 5 && BindingInfo.GetNumOutputs() == 3);
-		TSkinningModeBinder<TNDIParamBinder<0, int32, NDI_RAW_FUNC_BINDER(UNiagaraDataInterfaceSkeletalMesh, GetTriCoordVertices)>>::Bind(this, BindingInfo, InstanceData, OutFunc);
+		TSkinningModeBinder<NDI_FUNC_BINDER(UNiagaraDataInterfaceSkeletalMesh, GetTriCoordVertices)>::Bind(this, BindingInfo, InstanceData, OutFunc);
 	}
 
 	check(InstanceData->Mesh);
@@ -303,12 +303,12 @@ template<typename FilterMode, typename AreaWeightingMode>
 void UNiagaraDataInterfaceSkeletalMesh::RandomTriCoord(FVectorVMContext& Context)
 {
 	SCOPE_CYCLE_COUNTER(STAT_NiagaraSkel_Sample);
-	FUserPtrHandler<FNDISkeletalMesh_InstanceData> InstData(Context);
+	VectorVM::FUserPtrHandler<FNDISkeletalMesh_InstanceData> InstData(Context);
 	checkfSlow(InstData.Get(), TEXT("Skeletal Mesh Interface has invalid instance data. %s"), *GetPathName());
 	checkfSlow(InstData->Mesh, TEXT("Skeletal Mesh Interface has invalid mesh. %s"), *GetPathName());
 
-	FRegisterHandler<int32> OutTri(Context);
-	FRegisterHandler<float> OutBaryX(Context);	FRegisterHandler<float> OutBaryY(Context);	FRegisterHandler<float> OutBaryZ(Context);
+	VectorVM::FExternalFuncRegisterHandler<int32> OutTri(Context);
+	VectorVM::FExternalFuncRegisterHandler<float> OutBaryX(Context);	VectorVM::FExternalFuncRegisterHandler<float> OutBaryY(Context);	VectorVM::FExternalFuncRegisterHandler<float> OutBaryZ(Context);
 
 	FSkeletalMeshAccessorHelper MeshAccessor;
 	MeshAccessor.Init<FilterMode, AreaWeightingMode>(InstData);
@@ -324,19 +324,21 @@ void UNiagaraDataInterfaceSkeletalMesh::RandomTriCoord(FVectorVMContext& Context
 	}
 }
 
-template<typename FilterMode, typename AreaWeightingMode, typename TriType, typename BaryXType, typename BaryYType, typename BaryZType >
+template<typename FilterMode, typename AreaWeightingMode>
 void UNiagaraDataInterfaceSkeletalMesh::IsValidTriCoord(FVectorVMContext& Context)
 {
 	SCOPE_CYCLE_COUNTER(STAT_NiagaraSkel_Sample);
 
-	TriType TriParam(Context);
-	BaryXType BaryXParam(Context);	BaryYType BaryYParam(Context);	BaryZType BaryZParam(Context);
+	VectorVM::FExternalFuncInputHandler<int32> TriParam(Context);
+	VectorVM::FExternalFuncInputHandler<float> BaryXParam(Context);
+	VectorVM::FExternalFuncInputHandler<float> BaryYParam(Context);
+	VectorVM::FExternalFuncInputHandler<float> BaryZParam(Context);
 
-	FUserPtrHandler<FNDISkeletalMesh_InstanceData> InstData(Context);
+	VectorVM::FUserPtrHandler<FNDISkeletalMesh_InstanceData> InstData(Context);
 	checkfSlow(InstData.Get(), TEXT("Skeletal Mesh Interface has invalid instance data. %s"), *GetPathName());
 	checkfSlow(InstData->Mesh, TEXT("Skeletal Mesh Interface has invalid mesh. %s"), *GetPathName());
 
-	FRegisterHandler<FNiagaraBool> OutValid(Context);
+	VectorVM::FExternalFuncRegisterHandler<FNiagaraBool> OutValid(Context);
 
 	FSkeletalMeshAccessorHelper MeshAccessor;
 	MeshAccessor.Init<FilterMode, AreaWeightingMode>(InstData);
@@ -446,11 +448,11 @@ template<typename FilterMode, typename AreaWeightingMode>
 void UNiagaraDataInterfaceSkeletalMesh::GetFilteredTriangleCount(FVectorVMContext& Context)
 {
 	SCOPE_CYCLE_COUNTER(STAT_NiagaraSkel_Sample);
-	FUserPtrHandler<FNDISkeletalMesh_InstanceData> InstData(Context);
+	VectorVM::FUserPtrHandler<FNDISkeletalMesh_InstanceData> InstData(Context);
 	checkfSlow(InstData.Get(), TEXT("Skeletal Mesh Interface has invalid instance data. %s"), *GetPathName());
 	checkfSlow(InstData->Mesh, TEXT("Skeletal Mesh Interface has invalid mesh. %s"), *GetPathName());
 
-	FRegisterHandler<int32> OutTri(Context);
+	VectorVM::FExternalFuncRegisterHandler<int32> OutTri(Context);
 
 	FSkeletalMeshAccessorHelper MeshAccessor;
 	MeshAccessor.Init<FilterMode, AreaWeightingMode>(InstData);
@@ -566,17 +568,17 @@ FORCEINLINE int32 UNiagaraDataInterfaceSkeletalMesh::GetSpecificTriangleAt<
 	return 0;
 }
 
-template<typename FilterMode, typename AreaWeightingMode, typename TriType>
+template<typename FilterMode, typename AreaWeightingMode>
 void UNiagaraDataInterfaceSkeletalMesh::GetFilteredTriangleAt(FVectorVMContext& Context)
 {
 	SCOPE_CYCLE_COUNTER(STAT_NiagaraSkel_Sample);
 
-	TriType TriParam(Context);
-	FUserPtrHandler<FNDISkeletalMesh_InstanceData> InstData(Context);
+	VectorVM::FExternalFuncInputHandler<int32> TriParam(Context);
+	VectorVM::FUserPtrHandler<FNDISkeletalMesh_InstanceData> InstData(Context);
 	checkfSlow(InstData.Get(), TEXT("Skeletal Mesh Interface has invalid instance data. %s"), *GetPathName());
 	checkfSlow(InstData->Mesh, TEXT("Skeletal Mesh Interface has invalid mesh. %s"), *GetPathName());
-	FRegisterHandler<int32> OutTri(Context);
-	FRegisterHandler<float> OutBaryX(Context);	FRegisterHandler<float> OutBaryY(Context);	FRegisterHandler<float> OutBaryZ(Context);
+	VectorVM::FExternalFuncRegisterHandler<int32> OutTri(Context);
+	VectorVM::FExternalFuncRegisterHandler<float> OutBaryX(Context);	VectorVM::FExternalFuncRegisterHandler<float> OutBaryY(Context);	VectorVM::FExternalFuncRegisterHandler<float> OutBaryZ(Context);
 
 	FSkeletalMeshAccessorHelper Accessor;
 	Accessor.Init<FilterMode, AreaWeightingMode>(InstData);
@@ -600,20 +602,19 @@ void UNiagaraDataInterfaceSkeletalMesh::GetFilteredTriangleAt(FVectorVMContext&
 	}
 }
 
-template<typename TriType, typename BaryXType, typename BaryYType, typename BaryZType>
 void UNiagaraDataInterfaceSkeletalMesh::GetTriCoordColor(FVectorVMContext& Context)
 {
 	SCOPE_CYCLE_COUNTER(STAT_NiagaraSkel_Sample);
-	TriType TriParam(Context);
-	BaryXType BaryXParam(Context);
-	BaryYType BaryYParam(Context);
-	BaryZType BaryZParam(Context);
-	FUserPtrHandler<FNDISkeletalMesh_InstanceData> InstData(Context);
+	VectorVM::FExternalFuncInputHandler<int32> TriParam(Context);
+	VectorVM::FExternalFuncInputHandler<float> BaryXParam(Context);
+	VectorVM::FExternalFuncInputHandler<float> BaryYParam(Context);
+	VectorVM::FExternalFuncInputHandler<float> BaryZParam(Context);
+	VectorVM::FUserPtrHandler<FNDISkeletalMesh_InstanceData> InstData(Context);
 
-	FRegisterHandler<float> OutColorR(Context);
-	FRegisterHandler<float> OutColorG(Context);
-	FRegisterHandler<float> OutColorB(Context);
-	FRegisterHandler<float> OutColorA(Context);
+	VectorVM::FExternalFuncRegisterHandler<float> OutColorR(Context);
+	VectorVM::FExternalFuncRegisterHandler<float> OutColorG(Context);
+	VectorVM::FExternalFuncRegisterHandler<float> OutColorB(Context);
+	VectorVM::FExternalFuncRegisterHandler<float> OutColorA(Context);
 
 	USkeletalMeshComponent* Comp = Cast<USkeletalMeshComponent>(InstData->Component.Get());
 	FSkinWeightVertexBuffer* SkinWeightBuffer;
@@ -651,20 +652,22 @@ void UNiagaraDataInterfaceSkeletalMesh::GetTriCoordColor(FVectorVMContext& Conte
 	}
 }
 
-template<typename VertexAccessorType, typename TriType, typename BaryXType, typename BaryYType, typename BaryZType, typename UVSetType>
+template<typename VertexAccessorType>
 void UNiagaraDataInterfaceSkeletalMesh::GetTriCoordUV(FVectorVMContext& Context)
 {
 	SCOPE_CYCLE_COUNTER(STAT_NiagaraSkel_Sample);
 	VertexAccessorType VertAccessor;
-	TriType TriParam(Context);
-	BaryXType BaryXParam(Context);	BaryYType BaryYParam(Context);	BaryZType BaryZParam(Context);
-	UVSetType UVSetParam(Context);
-	FUserPtrHandler<FNDISkeletalMesh_InstanceData> InstData(Context);
+	VectorVM::FExternalFuncInputHandler<int32> TriParam(Context);
+	VectorVM::FExternalFuncInputHandler<float> BaryXParam(Context);
+	VectorVM::FExternalFuncInputHandler<float> BaryYParam(Context);
+	VectorVM::FExternalFuncInputHandler<float> BaryZParam(Context);
+	VectorVM::FExternalFuncInputHandler<int32> UVSetParam(Context);
+	VectorVM::FUserPtrHandler<FNDISkeletalMesh_InstanceData> InstData(Context);
 
 	checkfSlow(InstData.Get(), TEXT("Skeletal Mesh Interface has invalid instance data. %s"), *GetPathName());
 	checkfSlow(InstData->Mesh, TEXT("Skeletal Mesh Interface has invalid mesh. %s"), *GetPathName());
 
-	FRegisterHandler<float> OutUVX(Context);	FRegisterHandler<float> OutUVY(Context);
+	VectorVM::FExternalFuncRegisterHandler<float> OutUVX(Context);	VectorVM::FExternalFuncRegisterHandler<float> OutUVY(Context);
 
 	USkeletalMeshComponent* Comp = Cast<USkeletalMeshComponent>(InstData->Component.Get());
 	FSkinWeightVertexBuffer* SkinWeightBuffer;
@@ -719,12 +722,12 @@ struct FGetTriCoodSkinnedDataOutputHandler
 	{
 	}
 
-	FRegisterHandler<float> PosX; FRegisterHandler<float> PosY; FRegisterHandler<float> PosZ;
-	FRegisterHandler<float> VelX; FRegisterHandler<float> VelY; FRegisterHandler<float> VelZ;
+	VectorVM::FExternalFuncRegisterHandler<float> PosX; VectorVM::FExternalFuncRegisterHandler<float> PosY; VectorVM::FExternalFuncRegisterHandler<float> PosZ;
+	VectorVM::FExternalFuncRegisterHandler<float> VelX; VectorVM::FExternalFuncRegisterHandler<float> VelY; VectorVM::FExternalFuncRegisterHandler<float> VelZ;
 
-	FRegisterHandler<float> NormX; FRegisterHandler<float> NormY; FRegisterHandler<float> NormZ;
-	FRegisterHandler<float> BinormX; FRegisterHandler<float> BinormY; FRegisterHandler<float> BinormZ;
-	FRegisterHandler<float> TangentX; FRegisterHandler<float> TangentY; FRegisterHandler<float> TangentZ;
+	VectorVM::FExternalFuncRegisterHandler<float> NormX; VectorVM::FExternalFuncRegisterHandler<float> NormY; VectorVM::FExternalFuncRegisterHandler<float> NormZ;
+	VectorVM::FExternalFuncRegisterHandler<float> BinormX; VectorVM::FExternalFuncRegisterHandler<float> BinormY; VectorVM::FExternalFuncRegisterHandler<float> BinormZ;
+	VectorVM::FExternalFuncRegisterHandler<float> TangentX; VectorVM::FExternalFuncRegisterHandler<float> TangentY; VectorVM::FExternalFuncRegisterHandler<float> TangentZ;
 
 	const bool bNeedsPosition;
 	const bool bNeedsVelocity;
@@ -768,19 +771,19 @@ struct FGetTriCoodSkinnedDataOutputHandler
 	}
 };
 
-template<typename SkinningHandlerType, typename TransformHandlerType, typename VertexAccessorType, typename TriType, typename BaryXType, typename BaryYType, typename BaryZType>
+template<typename SkinningHandlerType, typename TransformHandlerType, typename VertexAccessorType>
 void UNiagaraDataInterfaceSkeletalMesh::GetTriCoordSkinnedData(FVectorVMContext& Context)
 {
 	SCOPE_CYCLE_COUNTER(STAT_NiagaraSkel_Sample);
 	SkinningHandlerType SkinningHandler;
 	TransformHandlerType TransformHandler;
 	VertexAccessorType VertAccessor;
-	TriType TriParam(Context);
-	BaryXType BaryXParam(Context);
-	BaryYType BaryYParam(Context);
-	BaryZType BaryZParam(Context);
+	VectorVM::FExternalFuncInputHandler<int32> TriParam(Context);
+	VectorVM::FExternalFuncInputHandler<float> BaryXParam(Context);
+	VectorVM::FExternalFuncInputHandler<float> BaryYParam(Context);
+	VectorVM::FExternalFuncInputHandler<float> BaryZParam(Context);
 
-	FUserPtrHandler<FNDISkeletalMesh_InstanceData> InstData(Context);
+	VectorVM::FUserPtrHandler<FNDISkeletalMesh_InstanceData> InstData(Context);
 
 	checkfSlow(InstData.Get(), TEXT("Skeletal Mesh Interface has invalid instance data. %s"), *GetPathName());
 	checkfSlow(InstData->Mesh, TEXT("Skeletal Mesh Interface has invalid mesh. %s"), *GetPathName());
@@ -832,7 +835,7 @@ void UNiagaraDataInterfaceSkeletalMesh::GetTriCoordSkinnedData(FVectorVMContext&
 			Velocity = (Pos - Prev) * InvDt;
 			Output.SetVelocity(Velocity);
 		}
-		
+
 		if (Output.bNeedsNorm)
 		{
 			Normal = ((Pos1 - Pos2) ^ (Pos0 - Pos2)).GetSafeNormal();
@@ -883,21 +886,21 @@ void UNiagaraDataInterfaceSkeletalMesh::GetTriCoordSkinnedData(FVectorVMContext&
 	}
 }
 
-template<typename SkinningHandlerType, typename TriType>
+template<typename SkinningHandlerType>
 void UNiagaraDataInterfaceSkeletalMesh::GetTriCoordVertices(FVectorVMContext& Context)
 {
 	SCOPE_CYCLE_COUNTER(STAT_NiagaraSkel_Sample);
 	SkinningHandlerType SkinningHandler;
-	TriType TriParam(Context);
+	VectorVM::FExternalFuncInputHandler<int32> TriParam(Context);
 
-	FUserPtrHandler<FNDISkeletalMesh_InstanceData> InstData(Context);
+	VectorVM::FUserPtrHandler<FNDISkeletalMesh_InstanceData> InstData(Context);
 
 	checkfSlow(InstData.Get(), TEXT("Skeletal Mesh Interface has invalid instance data. %s"), *GetPathName());
 	checkfSlow(InstData->Mesh, TEXT("Skeletal Mesh Interface has invalid mesh. %s"), *GetPathName());
 
-	FRegisterHandler<int32> OutV0(Context);
-	FRegisterHandler<int32> OutV1(Context);
-	FRegisterHandler<int32> OutV2(Context);
+	VectorVM::FExternalFuncRegisterHandler<int32> OutV0(Context);
+	VectorVM::FExternalFuncRegisterHandler<int32> OutV1(Context);
+	VectorVM::FExternalFuncRegisterHandler<int32> OutV2(Context);
 
 	int32 Idx0; int32 Idx1; int32 Idx2;
 	FSkeletalMeshAccessorHelper Accessor;
diff --git a/Engine/Plugins/FX/Niagara/Source/Niagara/Private/NDISkeletalMesh_VertexSampling.cpp b/Engine/Plugins/FX/Niagara/Source/Niagara/Private/NDISkeletalMesh_VertexSampling.cpp
--- a/Engine/Plugins/FX/Niagara/Source/Niagara/Private/NDISkeletalMesh_VertexSampling.cpp
+++ b/Engine/Plugins/FX/Niagara/Source/Niagara/Private/NDISkeletalMesh_VertexSampling.cpp
@@ -14,13 +14,13 @@
 DECLARE_CYCLE_STAT(TEXT("Skel Mesh Vertex Sampling"), STAT_NiagaraSkel_Vertex_Sample, STATGROUP_Niagara);
 
 //Final binders for all static mesh interface functions.
-DEFINE_NDI_RAW_FUNC_BINDER(UNiagaraDataInterfaceSkeletalMesh, RandomVertex);
-DEFINE_NDI_RAW_FUNC_BINDER(UNiagaraDataInterfaceSkeletalMesh, GetVertexSkinnedData);
-DEFINE_NDI_RAW_FUNC_BINDER(UNiagaraDataInterfaceSkeletalMesh, GetVertexColor);
-DEFINE_NDI_RAW_FUNC_BINDER(UNiagaraDataInterfaceSkeletalMesh, GetVertexUV);
-DEFINE_NDI_RAW_FUNC_BINDER(UNiagaraDataInterfaceSkeletalMesh, IsValidVertex);
-DEFINE_NDI_RAW_FUNC_BINDER(UNiagaraDataInterfaceSkeletalMesh, GetFilteredVertexCount);
-DEFINE_NDI_RAW_FUNC_BINDER(UNiagaraDataInterfaceSkeletalMesh, GetFilteredVertexAt);
+DEFINE_NDI_FUNC_BINDER(UNiagaraDataInterfaceSkeletalMesh, RandomVertex);
+DEFINE_NDI_FUNC_BINDER(UNiagaraDataInterfaceSkeletalMesh, GetVertexSkinnedData);
+DEFINE_NDI_DIRECT_FUNC_BINDER(UNiagaraDataInterfaceSkeletalMesh, GetVertexColor);
+DEFINE_NDI_FUNC_BINDER(UNiagaraDataInterfaceSkeletalMesh, GetVertexUV);
+DEFINE_NDI_FUNC_BINDER(UNiagaraDataInterfaceSkeletalMesh, IsValidVertex);
+DEFINE_NDI_FUNC_BINDER(UNiagaraDataInterfaceSkeletalMesh, GetFilteredVertexCount);
+DEFINE_NDI_FUNC_BINDER(UNiagaraDataInterfaceSkeletalMesh, GetFilteredVertexAt);
 
 static const FName RandomVertexName("RandomVertex");
 static const FName IsValidVertexName("IsValidVertex");
@@ -140,43 +140,43 @@ void UNiagaraDataInterfaceSkeletalMesh::BindVertexSamplingFunction(const FVMExte
 	if (BindingInfo.Name == RandomVertexName)
 	{
 		check(BindingInfo.GetNumInputs() == 1 && BindingInfo.GetNumOutputs() == 1);
-		TFilterModeBinder<NDI_RAW_FUNC_BINDER(UNiagaraDataInterfaceSkeletalMesh, RandomVertex)>::Bind(this, BindingInfo, InstanceData, OutFunc);
+		TFilterModeBinder<NDI_FUNC_BINDER(UNiagaraDataInterfaceSkeletalMesh, RandomVertex)>::Bind(this, BindingInfo, InstanceData, OutFunc);
 	}
 	else if (BindingInfo.Name == IsValidVertexName)
 	{
 		check(BindingInfo.GetNumInputs() == 2 && BindingInfo.GetNumOutputs() == 1);
-		TFilterModeBinder<TNDIParamBinder<0, int32, NDI_RAW_FUNC_BINDER(UNiagaraDataInterfaceSkeletalMesh, IsValidVertex)>>::Bind(this, BindingInfo, InstanceData, OutFunc);
+		TFilterModeBinder<NDI_FUNC_BINDER(UNiagaraDataInterfaceSkeletalMesh, IsValidVertex)>::Bind(this, BindingInfo, InstanceData, OutFunc);
 	}
 	else if (BindingInfo.Name == GetSkinnedVertexDataName)
 	{
 		check(BindingInfo.GetNumInputs() == 2 && BindingInfo.GetNumOutputs() == 6);
-		TSkinningModeBinder<TNDIExplicitBinder<FNDITransformHandlerNoop, TVertexAccessorBinder<TNDIParamBinder<0, int32, NDI_RAW_FUNC_BINDER(UNiagaraDataInterfaceSkeletalMesh, GetVertexSkinnedData)>>>>::Bind(this, BindingInfo, InstanceData, OutFunc);
+		TSkinningModeBinder<TNDIExplicitBinder<FNDITransformHandlerNoop, TVertexAccessorBinder<NDI_FUNC_BINDER(UNiagaraDataInterfaceSkeletalMesh, GetVertexSkinnedData)>>>::Bind(this, BindingInfo, InstanceData, OutFunc);
 	}
 	else if (BindingInfo.Name == GetSkinnedVertexDataWSName)
 	{
 		check(BindingInfo.GetNumInputs() == 2 && BindingInfo.GetNumOutputs() == 6);
-		TSkinningModeBinder<TNDIExplicitBinder<FNDITransformHandler, TVertexAccessorBinder<TNDIParamBinder<0, int32, NDI_RAW_FUNC_BINDER(UNiagaraDataInterfaceSkeletalMesh, GetVertexSkinnedData)>>>>::Bind(this, BindingInfo, InstanceData, OutFunc);
+		TSkinningModeBinder<TNDIExplicitBinder<FNDITransformHandler, TVertexAccessorBinder<NDI_FUNC_BINDER(UNiagaraDataInterfaceSkeletalMesh, GetVertexSkinnedData)>>>::Bind(this, BindingInfo, InstanceData, OutFunc);
 	}
 	else if (BindingInfo.Name == GetVertexColorName)
 	{
 		check(BindingInfo.GetNumInputs() == 2 && BindingInfo.GetNumOutputs() == 4);
 		bNeedsVertexColors = true;
-		TNDIParamBinder<0, int32, NDI_RAW_FUNC_BINDER(UNiagaraDataInterfaceSkeletalMesh, GetVertexColor)>::Bind(this, BindingInfo, InstanceData, OutFunc);
+		NDI_FUNC_BINDER(UNiagaraDataInterfaceSkeletalMesh, GetVertexColor)::Bind(this, OutFunc);
 	}
 	else if (BindingInfo.Name == GetVertexUVName)
 	{
 		check(BindingInfo.GetNumInputs() == 3 && BindingInfo.GetNumOutputs() == 2);
-		TVertexAccessorBinder<TNDIParamBinder<0, int32, TNDIParamBinder<1, int32, NDI_RAW_FUNC_BINDER(UNiagaraDataInterfaceSkeletalMesh, GetVertexUV)>>>::Bind(this, BindingInfo, InstanceData, OutFunc);
+		TVertexAccessorBinder<NDI_FUNC_BINDER(UNiagaraDataInterfaceSkeletalMesh, GetVertexUV)>::Bind(this, BindingInfo, InstanceData, OutFunc);
 	}
 	else if (BindingInfo.Name == GetVertexCountName)
 	{
 		check(BindingInfo.GetNumInputs() == 1 && BindingInfo.GetNumOutputs() == 1);
-		TFilterModeBinder<NDI_RAW_FUNC_BINDER(UNiagaraDataInterfaceSkeletalMesh, GetFilteredVertexCount)>::Bind(this, BindingInfo, InstanceData, OutFunc);
+		TFilterModeBinder<NDI_FUNC_BINDER(UNiagaraDataInterfaceSkeletalMesh, GetFilteredVertexCount)>::Bind(this, BindingInfo, InstanceData, OutFunc);
 	}
 	else if (BindingInfo.Name == GetVertexAtName)
 	{
 		check(BindingInfo.GetNumInputs() == 2 && BindingInfo.GetNumOutputs() == 1);
-		TFilterModeBinder<TNDIParamBinder<0, int32, NDI_RAW_FUNC_BINDER(UNiagaraDataInterfaceSkeletalMesh, GetFilteredVertexAt)>>::Bind(this, BindingInfo, InstanceData, OutFunc);
+		TFilterModeBinder<NDI_FUNC_BINDER(UNiagaraDataInterfaceSkeletalMesh, GetFilteredVertexAt)>::Bind(this, BindingInfo, InstanceData, OutFunc);
 	}
 
 	check(InstanceData->Mesh);
@@ -228,11 +228,11 @@ template<typename FilterMode>
 void UNiagaraDataInterfaceSkeletalMesh::RandomVertex(FVectorVMContext& Context)
 {
 	SCOPE_CYCLE_COUNTER(STAT_NiagaraSkel_Vertex_Sample);
-	FUserPtrHandler<FNDISkeletalMesh_InstanceData> InstData(Context);
+	VectorVM::FUserPtrHandler<FNDISkeletalMesh_InstanceData> InstData(Context);
 	checkfSlow(InstData.Get(), TEXT("Skeletal Mesh Interface has invalid instance data. %s"), *GetPathName());
 	checkfSlow(InstData->Mesh, TEXT("Skeletal Mesh Interface has invalid mesh. %s"), *GetPathName());
 
-	FRegisterHandler<int32> OutVert(Context);
+	VectorVM::FExternalFuncRegisterHandler<int32> OutVert(Context);
 
 	FSkeletalMeshAccessorHelper MeshAccessor;
 	MeshAccessor.Init<FilterMode, TIntegralConstant<int32,0>>(InstData);
@@ -243,18 +243,18 @@ void UNiagaraDataInterfaceSkeletalMesh::RandomVertex(FVectorVMContext& Context)
 	}
 }
 
-template<typename FilterMode, typename VertexType>
+template<typename FilterMode>
 void UNiagaraDataInterfaceSkeletalMesh::IsValidVertex(FVectorVMContext& Context)
 {
 	SCOPE_CYCLE_COUNTER(STAT_NiagaraSkel_Vertex_Sample);
 
-	VertexType VertexParam(Context);
+	VectorVM::FExternalFuncInputHandler<int32> VertexParam(Context);
 
-	FUserPtrHandler<FNDISkeletalMesh_InstanceData> InstData(Context);
+	VectorVM::FUserPtrHandler<FNDISkeletalMesh_InstanceData> InstData(Context);
 	checkfSlow(InstData.Get(), TEXT("Skeletal Mesh Interface has invalid instance data. %s"), *GetPathName());
 	checkfSlow(InstData->Mesh, TEXT("Skeletal Mesh Interface has invalid mesh. %s"), *GetPathName());
 
-	FRegisterHandler<FNiagaraBool> OutValid(Context);
+	VectorVM::FExternalFuncRegisterHandler<FNiagaraBool> OutValid(Context);
 
 	FSkeletalMeshAccessorHelper MeshAccessor;
 	MeshAccessor.Init<FilterMode, TIntegralConstant<ENDISkelMesh_AreaWeightingMode, ENDISkelMesh_AreaWeightingMode::None>>(InstData);
@@ -314,11 +314,11 @@ template<typename FilterMode>
 void UNiagaraDataInterfaceSkeletalMesh::GetFilteredVertexCount(FVectorVMContext& Context)
 {
 	SCOPE_CYCLE_COUNTER(STAT_NiagaraSkel_Vertex_Sample);
-	FUserPtrHandler<FNDISkeletalMesh_InstanceData> InstData(Context);
+	VectorVM::FUserPtrHandler<FNDISkeletalMesh_InstanceData> InstData(Context);
 	checkfSlow(InstData.Get(), TEXT("Skeletal Mesh Interface has invalid instance data. %s"), *GetPathName());
 	checkfSlow(InstData->Mesh, TEXT("Skeletal Mesh Interface has invalid mesh. %s"), *GetPathName());
 
-	FRegisterHandler<int32> OutVert(Context);
+	VectorVM::FExternalFuncRegisterHandler<int32> OutVert(Context);
 
 	FSkeletalMeshAccessorHelper MeshAccessor;
 	MeshAccessor.Init<FilterMode, TIntegralConstant<ENDISkelMesh_AreaWeightingMode, ENDISkelMesh_AreaWeightingMode::None>>(InstData);
@@ -375,16 +375,16 @@ FORCEINLINE_DEBUGGABLE int32 UNiagaraDataInterfaceSkeletalMesh::GetSpecificVerte
 	return 0;
 }
 
-template<typename FilterMode, typename VertexType>
+template<typename FilterMode>
 void UNiagaraDataInterfaceSkeletalMesh::GetFilteredVertexAt(FVectorVMContext& Context)
 {
 	SCOPE_CYCLE_COUNTER(STAT_NiagaraSkel_Vertex_Sample);
 
-	VertexType FilteredVertexParam(Context);
-	FUserPtrHandler<FNDISkeletalMesh_InstanceData> InstData(Context);
+	VectorVM::FExternalFuncInputHandler<int32> FilteredVertexParam(Context);
+	VectorVM::FUserPtrHandler<FNDISkeletalMesh_InstanceData> InstData(Context);
 	checkfSlow(InstData.Get(), TEXT("Skeletal Mesh Interface has invalid instance data. %s"), *GetPathName());
 	checkfSlow(InstData->Mesh, TEXT("Skeletal Mesh Interface has invalid mesh. %s"), *GetPathName());
-	FRegisterHandler<int32> OutVert(Context);
+	VectorVM::FExternalFuncRegisterHandler<int32> OutVert(Context);
 
 	FSkeletalMeshAccessorHelper Accessor;
 	Accessor.Init<FilterMode, TIntegralConstant<ENDISkelMesh_AreaWeightingMode, ENDISkelMesh_AreaWeightingMode::None>>(InstData);
@@ -401,17 +401,16 @@ void UNiagaraDataInterfaceSkeletalMesh::GetFilteredVertexAt(FVectorVMContext& Co
 	}
 }
 
-template<typename VertexType>
 void UNiagaraDataInterfaceSkeletalMesh::GetVertexColor(FVectorVMContext& Context)
 {
 	SCOPE_CYCLE_COUNTER(STAT_NiagaraSkel_Vertex_Sample);
-	VertexType VertParam(Context);
-	FUserPtrHandler<FNDISkeletalMesh_InstanceData> InstData(Context);
+	VectorVM::FExternalFuncInputHandler<int32> VertParam(Context);
+	VectorVM::FUserPtrHandler<FNDISkeletalMesh_InstanceData> InstData(Context);
 
-	FRegisterHandler<float> OutColorR(Context);
-	FRegisterHandler<float> OutColorG(Context);
-	FRegisterHandler<float> OutColorB(Context);
-	FRegisterHandler<float> OutColorA(Context);
+	VectorVM::FExternalFuncRegisterHandler<float> OutColorR(Context);
+	VectorVM::FExternalFuncRegisterHandler<float> OutColorG(Context);
+	VectorVM::FExternalFuncRegisterHandler<float> OutColorB(Context);
+	VectorVM::FExternalFuncRegisterHandler<float> OutColorA(Context);
 
 	USkeletalMeshComponent* Comp = Cast<USkeletalMeshComponent>(InstData->Component.Get());
 	FSkinWeightVertexBuffer* SkinWeightBuffer;
@@ -436,19 +435,19 @@ void UNiagaraDataInterfaceSkeletalMesh::GetVertexColor(FVectorVMContext& Context
 	}
 }
 
-template<typename VertexAccessorType, typename VertexType, typename UVSetType>
+template<typename VertexAccessorType>
 void UNiagaraDataInterfaceSkeletalMesh::GetVertexUV(FVectorVMContext& Context)
 {
 	SCOPE_CYCLE_COUNTER(STAT_NiagaraSkel_Vertex_Sample);
 	VertexAccessorType VertAccessor;
-	VertexType VertParam(Context);
-	UVSetType UVSetParam(Context);
-	FUserPtrHandler<FNDISkeletalMesh_InstanceData> InstData(Context);
+	VectorVM::FExternalFuncInputHandler<int32> VertParam(Context);
+	VectorVM::FExternalFuncInputHandler<int32> UVSetParam(Context);
+	VectorVM::FUserPtrHandler<FNDISkeletalMesh_InstanceData> InstData(Context);
 
 	checkfSlow(InstData.Get(), TEXT("Skeletal Mesh Interface has invalid instance data. %s"), *GetPathName());
 	checkfSlow(InstData->Mesh, TEXT("Skeletal Mesh Interface has invalid mesh. %s"), *GetPathName());
 
-	FRegisterHandler<float> OutUVX(Context);	FRegisterHandler<float> OutUVY(Context);
+	VectorVM::FExternalFuncRegisterHandler<float> OutUVX(Context);	VectorVM::FExternalFuncRegisterHandler<float> OutUVY(Context);
 
 	USkeletalMeshComponent* Comp = Cast<USkeletalMeshComponent>(InstData->Component.Get());
 	FSkinWeightVertexBuffer* SkinWeightBuffer;
@@ -480,8 +479,8 @@ struct FGetVertexSkinnedDataOutputHandler
 	{
 	}
 
-	FRegisterHandler<float> PosX; FRegisterHandler<float> PosY; FRegisterHandler<float> PosZ;
-	FRegisterHandler<float> VelX; FRegisterHandler<float> VelY; FRegisterHandler<float> VelZ;
+	VectorVM::FExternalFuncRegisterHandler<float> PosX; VectorVM::FExternalFuncRegisterHandler<float> PosY; VectorVM::FExternalFuncRegisterHandler<float> PosZ;
+	VectorVM::FExternalFuncRegisterHandler<float> VelX; VectorVM::FExternalFuncRegisterHandler<float> VelY; VectorVM::FExternalFuncRegisterHandler<float> VelZ;
 
 	const bool bNeedsPosition;
 	const bool bNeedsVelocity;
@@ -501,15 +500,15 @@ struct FGetVertexSkinnedDataOutputHandler
 	}
 };
 
-template<typename SkinningHandlerType, typename TransformHandlerType, typename VertexAccessorType, typename VertexType>
+template<typename SkinningHandlerType, typename TransformHandlerType, typename VertexAccessorType>
 void UNiagaraDataInterfaceSkeletalMesh::GetVertexSkinnedData(FVectorVMContext& Context)
 {
 	SCOPE_CYCLE_COUNTER(STAT_NiagaraSkel_Vertex_Sample);
 	SkinningHandlerType SkinningHandler;
 	TransformHandlerType TransformHandler;
-	VertexType VertParam(Context);
+	VectorVM::FExternalFuncInputHandler<int32> VertParam(Context);
 
-	FUserPtrHandler<FNDISkeletalMesh_InstanceData> InstData(Context);
+	VectorVM::FUserPtrHandler<FNDISkeletalMesh_InstanceData> InstData(Context);
 
 	checkfSlow(InstData.Get(), TEXT("Skeletal Mesh Interface has invalid instance data. %s"), *GetPathName());
 	checkfSlow(InstData->Mesh, TEXT("Skeletal Mesh Interface has invalid mesh. %s"), *GetPathName());
diff --git a/Engine/Plugins/FX/Niagara/Source/Niagara/Private/NiagaraDataInterfaceCollisionQuery.cpp b/Engine/Plugins/FX/Niagara/Source/Niagara/Private/NiagaraDataInterfaceCollisionQuery.cpp
--- a/Engine/Plugins/FX/Niagara/Source/Niagara/Private/NiagaraDataInterfaceCollisionQuery.cpp
+++ b/Engine/Plugins/FX/Niagara/Source/Niagara/Private/NiagaraDataInterfaceCollisionQuery.cpp
@@ -136,7 +136,7 @@ bool UNiagaraDataInterfaceCollisionQuery::GetFunctionHLSL(const FName& Definitio
 	{
 		OutHLSL += TEXT("float3 WorldPositionFromSceneDepth(float2 ScreenPosition, float SceneDepth)\n{\n\tfloat4 HomogeneousWorldPosition = mul(float4(ScreenPosition * SceneDepth, SceneDepth, 1), View.ScreenToWorld);\n\treturn HomogeneousWorldPosition.xyz / HomogeneousWorldPosition.w;\n}\n");
 		OutHLSL += TEXT("\n");
-		OutHLSL += TEXT("void ") + InstanceFunctionName + TEXT("(in int InQueryID, in float3 In_ParticlePos, in float3 In_ParticleVel, in float In_DeltaSeconds, float CollisionRadius, in float CollisionDepthBounds, int In_InstanceData, \
+		OutHLSL += TEXT("void ") + InstanceFunctionName + TEXT("(in int InQueryID, in float3 In_ParticlePos, in float3 In_ParticleVel, in float In_DeltaSeconds, float CollisionRadius, in float CollisionDepthBounds, \
 			out int Out_QueryID, out bool OutCollisionValid, out float3 Out_CollisionPos, out float3 Out_CollisionNormal, out float Out_Friction, out float Out_Restitution) \n{\n");
 		// get the screen position
 		OutHLSL += TEXT("\
@@ -201,25 +201,24 @@ void UNiagaraDataInterfaceCollisionQuery::GetParameterDefinitionHLSL(FNiagaraDat
 	// we don't need to add these to hlsl, as they're already in common.ush
 }
 
-DEFINE_NDI_RAW_FUNC_BINDER(UNiagaraDataInterfaceCollisionQuery, SubmitQuery);
-DEFINE_NDI_RAW_FUNC_BINDER(UNiagaraDataInterfaceCollisionQuery, ReadQuery);
-DEFINE_NDI_RAW_FUNC_BINDER(UNiagaraDataInterfaceCollisionQuery, PerformQuery);
+DEFINE_NDI_DIRECT_FUNC_BINDER(UNiagaraDataInterfaceCollisionQuery, SubmitQuery);
+DEFINE_NDI_DIRECT_FUNC_BINDER(UNiagaraDataInterfaceCollisionQuery, ReadQuery);
+DEFINE_NDI_DIRECT_FUNC_BINDER(UNiagaraDataInterfaceCollisionQuery, PerformQuery);
 
 void UNiagaraDataInterfaceCollisionQuery::GetVMExternalFunction(const FVMExternalFunctionBindingInfo& BindingInfo, void* InstanceData, FVMExternalFunction &OutFunc)
 {
 	CQDIPerInstanceData *InstData = (CQDIPerInstanceData *)InstanceData;
 	if (BindingInfo.Name == TEXT("SubmitQuery") /*&& BindingInfo.GetNumInputs() == 3 && BindingInfo.GetNumOutputs() == 1*/)
 	{
-		TNDIParamBinder<0, float, TNDIParamBinder<1, float, TNDIParamBinder<2, float, TNDIParamBinder<3, float, TNDIParamBinder<4, float, TNDIParamBinder<5, float, TNDIParamBinder<6, float, NDI_RAW_FUNC_BINDER(UNiagaraDataInterfaceCollisionQuery, SubmitQuery)>>>>>>>::Bind(this, BindingInfo, InstData, OutFunc);
+		 NDI_FUNC_BINDER(UNiagaraDataInterfaceCollisionQuery, SubmitQuery)::Bind(this, OutFunc);
 	}
 	else if (BindingInfo.Name == TEXT("ReadQuery") /*&& BindingInfo.GetNumInputs() == 1 && BindingInfo.GetNumOutputs() == 4*/)
 	{
-		TNDIParamBinder<0, int32, NDI_RAW_FUNC_BINDER(UNiagaraDataInterfaceCollisionQuery, ReadQuery)>::Bind(this, BindingInfo, InstData, OutFunc);
+		NDI_FUNC_BINDER(UNiagaraDataInterfaceCollisionQuery, ReadQuery)::Bind(this, OutFunc);
 	}
 	else if (BindingInfo.Name == TEXT("PerformCollisionQuery") /*&& BindingInfo.GetNumInputs() == 1 && BindingInfo.GetNumOutputs() == 4*/)
 	{
-		//				  id     |---------------------------------- position ---------------------------|    |------------------------------- direction -------------------------------|     |-------- dT --------|   |------- size -------|    |---- depthBounds ----|
-		TNDIParamBinder<0, int, TNDIParamBinder<1, float, TNDIParamBinder<2, float, TNDIParamBinder<3, float, TNDIParamBinder<4, float, TNDIParamBinder<5, float, TNDIParamBinder<6, float, TNDIParamBinder<7, float, TNDIParamBinder<8, float, TNDIParamBinder<9, float, NDI_RAW_FUNC_BINDER(UNiagaraDataInterfaceCollisionQuery, PerformQuery)>>>>>>>>>>::Bind(this, BindingInfo, InstData, OutFunc);
+		NDI_FUNC_BINDER(UNiagaraDataInterfaceCollisionQuery, PerformQuery)::Bind(this, OutFunc);
 	}
 	else
 	{
@@ -228,37 +227,35 @@ void UNiagaraDataInterfaceCollisionQuery::GetVMExternalFunction(const FVMExterna
 	}
 }
 
-
-template<typename InQueryIDType, typename PosTypeX, typename PosTypeY, typename PosTypeZ, typename DirTypeX, typename DirTypeY, typename DirTypeZ, typename DTType, typename SizeType, typename DepthBoundsType>
 void UNiagaraDataInterfaceCollisionQuery::PerformQuery(FVectorVMContext& Context)
 {
 	//PerformType PerformParam(Context);
-	InQueryIDType InIDParam(Context);
-	PosTypeX PosParamX(Context);
-	PosTypeY PosParamY(Context);
-	PosTypeZ PosParamZ(Context);
+	VectorVM::FExternalFuncInputHandler<int32> InIDParam(Context);
+	VectorVM::FExternalFuncInputHandler<float> PosParamX(Context);
+	VectorVM::FExternalFuncInputHandler<float> PosParamY(Context);
+	VectorVM::FExternalFuncInputHandler<float> PosParamZ(Context);
 
-	DirTypeX DirParamX(Context);
-	DirTypeY DirParamY(Context);
-	DirTypeZ DirParamZ(Context);
+	VectorVM::FExternalFuncInputHandler<float> DirParamX(Context);
+	VectorVM::FExternalFuncInputHandler<float> DirParamY(Context);
+	VectorVM::FExternalFuncInputHandler<float> DirParamZ(Context);
 
-	DTType DTParam(Context);
-	SizeType SizeParam(Context);
-	DepthBoundsType DepthBoundsParam(Context);
+	VectorVM::FExternalFuncInputHandler<float> DTParam(Context);
+	VectorVM::FExternalFuncInputHandler<float> SizeParam(Context);
+	VectorVM::FExternalFuncInputHandler<float> DepthBoundsParam(Context);
 
-	FUserPtrHandler<CQDIPerInstanceData> InstanceData(Context);
+	VectorVM::FUserPtrHandler<CQDIPerInstanceData> InstanceData(Context);
 
-	FRegisterHandler<int32> OutQueryID(Context);
+	VectorVM::FExternalFuncRegisterHandler<int32> OutQueryID(Context);
 
-	FRegisterHandler<int32> OutQueryValid(Context);
-	FRegisterHandler<float> OutCollisionPosX(Context);
-	FRegisterHandler<float> OutCollisionPosY(Context);
-	FRegisterHandler<float> OutCollisionPosZ(Context);
-	FRegisterHandler<float> OutCollisionNormX(Context);
-	FRegisterHandler<float> OutCollisionNormY(Context);
-	FRegisterHandler<float> OutCollisionNormZ(Context);
-	FRegisterHandler<float> Friction(Context);
-	FRegisterHandler<float> Restitution(Context);
+	VectorVM::FExternalFuncRegisterHandler<int32> OutQueryValid(Context);
+	VectorVM::FExternalFuncRegisterHandler<float> OutCollisionPosX(Context);
+	VectorVM::FExternalFuncRegisterHandler<float> OutCollisionPosY(Context);
+	VectorVM::FExternalFuncRegisterHandler<float> OutCollisionPosZ(Context);
+	VectorVM::FExternalFuncRegisterHandler<float> OutCollisionNormX(Context);
+	VectorVM::FExternalFuncRegisterHandler<float> OutCollisionNormY(Context);
+	VectorVM::FExternalFuncRegisterHandler<float> OutCollisionNormZ(Context);
+	VectorVM::FExternalFuncRegisterHandler<float> Friction(Context);
+	VectorVM::FExternalFuncRegisterHandler<float> Restitution(Context);
 
 	FScopeLock ScopeLock(&CriticalSection);
 	for (int32 i = 0; i < Context.NumInstances; ++i)
@@ -330,22 +327,21 @@ void UNiagaraDataInterfaceCollisionQuery::PerformQuery(FVectorVMContext& Context
 }
 
 
-template<typename PosTypeX, typename PosTypeY, typename PosTypeZ, typename VelTypeX, typename VelTypeY, typename VelTypeZ, typename DTType>
 void UNiagaraDataInterfaceCollisionQuery::SubmitQuery(FVectorVMContext& Context)
 {
-	PosTypeX PosParamX(Context);
-	PosTypeY PosParamY(Context);
-	PosTypeZ PosParamZ(Context);
-	VelTypeX VelParamX(Context);
-	VelTypeY VelParamY(Context);
-	VelTypeZ VelParamZ(Context);
+	VectorVM::FExternalFuncInputHandler<float> PosParamX(Context);
+	VectorVM::FExternalFuncInputHandler<float> PosParamY(Context);
+	VectorVM::FExternalFuncInputHandler<float> PosParamZ(Context);
+	VectorVM::FExternalFuncInputHandler<float> VelParamX(Context);
+	VectorVM::FExternalFuncInputHandler<float> VelParamY(Context);
+	VectorVM::FExternalFuncInputHandler<float> VelParamZ(Context);
 
-	DTType DTParam(Context);
-	FUserPtrHandler<CQDIPerInstanceData> InstanceData(Context);
+	VectorVM::FExternalFuncInputHandler<float> DTParam(Context);
+	VectorVM::FUserPtrHandler<CQDIPerInstanceData> InstanceData(Context);
 
 	FScopeLock ScopeLock(&CriticalSection);
 
-	FRegisterHandler<int32> OutQueryID(Context);
+	VectorVM::FExternalFuncRegisterHandler<int32> OutQueryID(Context);
 	for (int32 i = 0; i < Context.NumInstances; ++i)
 	{
 		FVector Pos(PosParamX.Get(), PosParamY.Get(), PosParamZ.Get());
@@ -368,22 +364,20 @@ void UNiagaraDataInterfaceCollisionQuery::SubmitQuery(FVectorVMContext& Context)
 
 }
 
-
-template<typename IDType>
 void UNiagaraDataInterfaceCollisionQuery::ReadQuery(FVectorVMContext& Context)
 {	
-	IDType IDParam(Context);
-	FUserPtrHandler<CQDIPerInstanceData> InstanceData(Context);
-	FRegisterHandler<int32> OutQueryValid(Context);
-	FRegisterHandler<float> OutCollisionPosX(Context);
-	FRegisterHandler<float> OutCollisionPosY(Context);
-	FRegisterHandler<float> OutCollisionPosZ(Context);
-	FRegisterHandler<float> OutCollisionVelX(Context);
-	FRegisterHandler<float> OutCollisionVelY(Context);
-	FRegisterHandler<float> OutCollisionVelZ(Context);
-	FRegisterHandler<float> OutCollisionNormX(Context);
-	FRegisterHandler<float> OutCollisionNormY(Context);
-	FRegisterHandler<float> OutCollisionNormZ(Context);
+	VectorVM::FExternalFuncInputHandler<int32> IDParam(Context);
+	VectorVM::FUserPtrHandler<CQDIPerInstanceData> InstanceData(Context);
+	VectorVM::FExternalFuncRegisterHandler<int32> OutQueryValid(Context);
+	VectorVM::FExternalFuncRegisterHandler<float> OutCollisionPosX(Context);
+	VectorVM::FExternalFuncRegisterHandler<float> OutCollisionPosY(Context);
+	VectorVM::FExternalFuncRegisterHandler<float> OutCollisionPosZ(Context);
+	VectorVM::FExternalFuncRegisterHandler<float> OutCollisionVelX(Context);
+	VectorVM::FExternalFuncRegisterHandler<float> OutCollisionVelY(Context);
+	VectorVM::FExternalFuncRegisterHandler<float> OutCollisionVelZ(Context);
+	VectorVM::FExternalFuncRegisterHandler<float> OutCollisionNormX(Context);
+	VectorVM::FExternalFuncRegisterHandler<float> OutCollisionNormY(Context);
+	VectorVM::FExternalFuncRegisterHandler<float> OutCollisionNormZ(Context);
 
 	FScopeLock ScopeLock(&CriticalSection);
 	for (int32 i = 0; i < Context.NumInstances; ++i)
diff --git a/Engine/Plugins/FX/Niagara/Source/Niagara/Private/NiagaraDataInterfaceColorCurve.cpp b/Engine/Plugins/FX/Niagara/Source/Niagara/Private/NiagaraDataInterfaceColorCurve.cpp
--- a/Engine/Plugins/FX/Niagara/Source/Niagara/Private/NiagaraDataInterfaceColorCurve.cpp
+++ b/Engine/Plugins/FX/Niagara/Source/Niagara/Private/NiagaraDataInterfaceColorCurve.cpp
@@ -175,12 +175,12 @@ void %s(in float In_X, out float4 Out_Value) \n\
 	return true;
 }
 
-DEFINE_NDI_RAW_FUNC_BINDER(UNiagaraDataInterfaceColorCurve, SampleCurve);
+DEFINE_NDI_FUNC_BINDER(UNiagaraDataInterfaceColorCurve, SampleCurve);
 void UNiagaraDataInterfaceColorCurve::GetVMExternalFunction(const FVMExternalFunctionBindingInfo& BindingInfo, void* InstanceData, FVMExternalFunction &OutFunc)
 {
 	if (BindingInfo.Name == SampleCurveName && BindingInfo.GetNumInputs() == 1 && BindingInfo.GetNumOutputs() == 4)
 	{
-		TCurveUseLUTBinder<TNDIParamBinder<0, float, NDI_RAW_FUNC_BINDER(UNiagaraDataInterfaceColorCurve, SampleCurve)>>::Bind(this, BindingInfo, InstanceData, OutFunc);
+		TCurveUseLUTBinder<NDI_FUNC_BINDER(UNiagaraDataInterfaceColorCurve, SampleCurve)>::Bind(this, BindingInfo, InstanceData, OutFunc);
 	}
 	else
 	{
@@ -211,15 +211,15 @@ FORCEINLINE_DEBUGGABLE FLinearColor UNiagaraDataInterfaceColorCurve::SampleCurve
 	return FLinearColor(RedCurve.Eval(X), GreenCurve.Eval(X), BlueCurve.Eval(X), AlphaCurve.Eval(X));
 }
 
-template<typename UseLUT, typename XParamType>
+template<typename UseLUT>
 void UNiagaraDataInterfaceColorCurve::SampleCurve(FVectorVMContext& Context)
 {
 	//TODO: Create some SIMDable optimized representation of the curve to do this faster.
-	XParamType XParam(Context);
-	FRegisterHandler<float> SamplePtrR(Context);
-	FRegisterHandler<float> SamplePtrG(Context);
-	FRegisterHandler<float> SamplePtrB(Context);
-	FRegisterHandler<float> SamplePtrA(Context);
+	VectorVM::FExternalFuncInputHandler<float> XParam(Context);
+	VectorVM::FExternalFuncRegisterHandler<float> SamplePtrR(Context);
+	VectorVM::FExternalFuncRegisterHandler<float> SamplePtrG(Context);
+	VectorVM::FExternalFuncRegisterHandler<float> SamplePtrB(Context);
+	VectorVM::FExternalFuncRegisterHandler<float> SamplePtrA(Context);
 
 	for (int32 i = 0; i < Context.NumInstances; ++i)
 	{
diff --git a/Engine/Plugins/FX/Niagara/Source/Niagara/Private/NiagaraDataInterfaceCurlNoise.cpp b/Engine/Plugins/FX/Niagara/Source/Niagara/Private/NiagaraDataInterfaceCurlNoise.cpp
--- a/Engine/Plugins/FX/Niagara/Source/Niagara/Private/NiagaraDataInterfaceCurlNoise.cpp
+++ b/Engine/Plugins/FX/Niagara/Source/Niagara/Private/NiagaraDataInterfaceCurlNoise.cpp
@@ -82,23 +82,22 @@ void UNiagaraDataInterfaceCurlNoise::GetFunctions(TArray<FNiagaraFunctionSignatu
 	OutFunctions.Add(Sig);
 }
 
-DEFINE_NDI_RAW_FUNC_BINDER(UNiagaraDataInterfaceCurlNoise, SampleNoiseField);
+DEFINE_NDI_DIRECT_FUNC_BINDER(UNiagaraDataInterfaceCurlNoise, SampleNoiseField);
 void UNiagaraDataInterfaceCurlNoise::GetVMExternalFunction(const FVMExternalFunctionBindingInfo& BindingInfo, void* InstanceData, FVMExternalFunction &OutFunc)
 {
 	check(BindingInfo.Name == SampleNoiseFieldName);
 	check(BindingInfo.GetNumInputs() == 3 && BindingInfo.GetNumOutputs() == 3);
-	TNDIParamBinder<0, float, TNDIParamBinder<1, float, TNDIParamBinder<2, float, NDI_RAW_FUNC_BINDER(UNiagaraDataInterfaceCurlNoise, SampleNoiseField)>>>::Bind(this, BindingInfo, InstanceData, OutFunc);
+	NDI_FUNC_BINDER(UNiagaraDataInterfaceCurlNoise, SampleNoiseField)::Bind(this, OutFunc);
 }
 
-template<typename XType, typename YType, typename ZType>
 void UNiagaraDataInterfaceCurlNoise::SampleNoiseField(FVectorVMContext& Context)
 {
-	XType XParam(Context);
-	YType YParam(Context);
-	ZType ZParam(Context);
-	FRegisterHandler<float> OutSampleX(Context);
-	FRegisterHandler<float> OutSampleY(Context);
-	FRegisterHandler<float> OutSampleZ(Context);
+	VectorVM::FExternalFuncInputHandler<float> XParam(Context);
+	VectorVM::FExternalFuncInputHandler<float> YParam(Context);
+	VectorVM::FExternalFuncInputHandler<float> ZParam(Context);
+	VectorVM::FExternalFuncRegisterHandler<float> OutSampleX(Context);
+	VectorVM::FExternalFuncRegisterHandler<float> OutSampleY(Context);
+	VectorVM::FExternalFuncRegisterHandler<float> OutSampleZ(Context);
 
 	const VectorRegister One = MakeVectorRegister(1.0f, 1.0f, 1.0f, 1.0f);
 	const VectorRegister Zero = MakeVectorRegister(0.0f, 0.0f, 0.0f, 0.0f);
diff --git a/Engine/Plugins/FX/Niagara/Source/Niagara/Private/NiagaraDataInterfaceCurve.cpp b/Engine/Plugins/FX/Niagara/Source/Niagara/Private/NiagaraDataInterfaceCurve.cpp
--- a/Engine/Plugins/FX/Niagara/Source/Niagara/Private/NiagaraDataInterfaceCurve.cpp
+++ b/Engine/Plugins/FX/Niagara/Source/Niagara/Private/NiagaraDataInterfaceCurve.cpp
@@ -144,12 +144,12 @@ void %s(in float In_X, out float Out_Value) \n\
 	return true;
 }
 
-DEFINE_NDI_RAW_FUNC_BINDER(UNiagaraDataInterfaceCurve, SampleCurve);
+DEFINE_NDI_FUNC_BINDER(UNiagaraDataInterfaceCurve, SampleCurve);
 void UNiagaraDataInterfaceCurve::GetVMExternalFunction(const FVMExternalFunctionBindingInfo& BindingInfo, void* InstanceData, FVMExternalFunction &OutFunc)
 {
 	if (BindingInfo.Name == SampleCurveName && BindingInfo.GetNumInputs() == 1 && BindingInfo.GetNumOutputs() == 1)
 	{
-		TCurveUseLUTBinder<TNDIParamBinder<0, float, NDI_RAW_FUNC_BINDER(UNiagaraDataInterfaceCurve, SampleCurve)>>::Bind(this, BindingInfo, InstanceData, OutFunc);
+		TCurveUseLUTBinder<NDI_FUNC_BINDER(UNiagaraDataInterfaceCurve, SampleCurve)>::Bind(this, BindingInfo, InstanceData, OutFunc);
 	}
 	else
 	{
@@ -179,11 +179,11 @@ FORCEINLINE_DEBUGGABLE float UNiagaraDataInterfaceCurve::SampleCurveInternal<TIn
 	return Curve.Eval(X);
 }
 
-template<typename UseLUT, typename XParamType>
+template<typename UseLUT>
 void UNiagaraDataInterfaceCurve::SampleCurve(FVectorVMContext& Context)
 {
-	XParamType XParam(Context);
-	FRegisterHandler<float> OutSample(Context);
+	VectorVM::FExternalFuncInputHandler<float> XParam(Context);
+	VectorVM::FExternalFuncRegisterHandler<float> OutSample(Context);
 
 	for (int32 i = 0; i < Context.NumInstances; ++i)
 	{
diff --git a/Engine/Plugins/FX/Niagara/Source/Niagara/Private/NiagaraDataInterfaceSimpleCounter.cpp b/Engine/Plugins/FX/Niagara/Source/Niagara/Private/NiagaraDataInterfaceSimpleCounter.cpp
--- a/Engine/Plugins/FX/Niagara/Source/Niagara/Private/NiagaraDataInterfaceSimpleCounter.cpp
+++ b/Engine/Plugins/FX/Niagara/Source/Niagara/Private/NiagaraDataInterfaceSimpleCounter.cpp
@@ -110,8 +110,8 @@ void UNiagaraDataInterfaceSimpleCounter::GetVMExternalFunction(const FVMExternal
 
 void UNiagaraDataInterfaceSimpleCounter::GetNextValue(FVectorVMContext& Context)
 {	
-	FUserPtrHandler<CounterInstanceData> InstanceData(Context);
-	FRegisterHandler<int32> OutValue(Context);
+	VectorVM::FUserPtrHandler<CounterInstanceData> InstanceData(Context);
+	VectorVM::FExternalFuncRegisterHandler<int32> OutValue(Context);
 
 	for (int32 i = 0; i < Context.NumInstances; ++i)
 	{
diff --git a/Engine/Plugins/FX/Niagara/Source/Niagara/Private/NiagaraDataInterfaceSpline.cpp b/Engine/Plugins/FX/Niagara/Source/Niagara/Private/NiagaraDataInterfaceSpline.cpp
--- a/Engine/Plugins/FX/Niagara/Source/Niagara/Private/NiagaraDataInterfaceSpline.cpp
+++ b/Engine/Plugins/FX/Niagara/Source/Niagara/Private/NiagaraDataInterfaceSpline.cpp
@@ -220,69 +220,69 @@ void UNiagaraDataInterfaceSpline::GetFunctions(TArray<FNiagaraFunctionSignature>
 	}
 }
 
-DEFINE_NDI_RAW_FUNC_BINDER(UNiagaraDataInterfaceSpline, SampleSplinePositionByUnitDistance);
-DEFINE_NDI_RAW_FUNC_BINDER(UNiagaraDataInterfaceSpline, SampleSplineUpVectorByUnitDistance);
-DEFINE_NDI_RAW_FUNC_BINDER(UNiagaraDataInterfaceSpline, SampleSplineRightVectorByUnitDistance);
-DEFINE_NDI_RAW_FUNC_BINDER(UNiagaraDataInterfaceSpline, SampleSplineDirectionByUnitDistance);
-DEFINE_NDI_RAW_FUNC_BINDER(UNiagaraDataInterfaceSpline, SampleSplineTangentByUnitDistance);
-DEFINE_NDI_RAW_FUNC_BINDER(UNiagaraDataInterfaceSpline, FindClosestUnitDistanceFromPositionWS);
+DEFINE_NDI_FUNC_BINDER(UNiagaraDataInterfaceSpline, SampleSplinePositionByUnitDistance);
+DEFINE_NDI_FUNC_BINDER(UNiagaraDataInterfaceSpline, SampleSplineUpVectorByUnitDistance);
+DEFINE_NDI_FUNC_BINDER(UNiagaraDataInterfaceSpline, SampleSplineRightVectorByUnitDistance);
+DEFINE_NDI_FUNC_BINDER(UNiagaraDataInterfaceSpline, SampleSplineDirectionByUnitDistance);
+DEFINE_NDI_FUNC_BINDER(UNiagaraDataInterfaceSpline, SampleSplineTangentByUnitDistance);
+DEFINE_NDI_FUNC_BINDER(UNiagaraDataInterfaceSpline, FindClosestUnitDistanceFromPositionWS);
 
 void UNiagaraDataInterfaceSpline::GetVMExternalFunction(const FVMExternalFunctionBindingInfo& BindingInfo, void* InstanceData, FVMExternalFunction &OutFunc)
 {
 	if (BindingInfo.Name == SampleSplinePositionByUnitDistanceName)
 	{
 		check(BindingInfo.GetNumInputs() == 2 && BindingInfo.GetNumOutputs() == 3);
-		TNDIExplicitBinder<FNDITransformHandlerNoop, TNDIParamBinder<0, float, NDI_RAW_FUNC_BINDER(UNiagaraDataInterfaceSpline, SampleSplinePositionByUnitDistance)>>::Bind(this, BindingInfo, InstanceData, OutFunc);
+		TNDIExplicitBinder<FNDITransformHandlerNoop, TNDIParamBinder<0, float, NDI_FUNC_BINDER(UNiagaraDataInterfaceSpline, SampleSplinePositionByUnitDistance)>>::Bind(this, BindingInfo, InstanceData, OutFunc);
 	}
 	else if (BindingInfo.Name == SampleSplinePositionByUnitDistanceWSName)
 	{
 		check(BindingInfo.GetNumInputs() == 2 && BindingInfo.GetNumOutputs() == 3);
-		TNDIExplicitBinder<FNDITransformHandler, TNDIParamBinder<0, float, NDI_RAW_FUNC_BINDER(UNiagaraDataInterfaceSpline, SampleSplinePositionByUnitDistance)>>::Bind(this, BindingInfo, InstanceData, OutFunc);
+		TNDIExplicitBinder<FNDITransformHandler, TNDIParamBinder<0, float, NDI_FUNC_BINDER(UNiagaraDataInterfaceSpline, SampleSplinePositionByUnitDistance)>>::Bind(this, BindingInfo, InstanceData, OutFunc);
 	}
 	else if (BindingInfo.Name == SampleSplineUpVectorByUnitDistanceName)
 	{
 		check(BindingInfo.GetNumInputs() == 2 && BindingInfo.GetNumOutputs() == 3);
-		TNDIExplicitBinder<FNDITransformHandlerNoop, TNDIParamBinder<0, float, NDI_RAW_FUNC_BINDER(UNiagaraDataInterfaceSpline, SampleSplineUpVectorByUnitDistance)>>::Bind(this, BindingInfo, InstanceData, OutFunc);
+		TNDIExplicitBinder<FNDITransformHandlerNoop, TNDIParamBinder<0, float, NDI_FUNC_BINDER(UNiagaraDataInterfaceSpline, SampleSplineUpVectorByUnitDistance)>>::Bind(this, BindingInfo, InstanceData, OutFunc);
 	}
 	else if (BindingInfo.Name == SampleSplineUpVectorByUnitDistanceWSName)
 	{
 		check(BindingInfo.GetNumInputs() == 2 && BindingInfo.GetNumOutputs() == 3);
-		TNDIExplicitBinder<FNDITransformHandler, TNDIParamBinder<0, float, NDI_RAW_FUNC_BINDER(UNiagaraDataInterfaceSpline, SampleSplineUpVectorByUnitDistance)>>::Bind(this, BindingInfo, InstanceData, OutFunc);
+		TNDIExplicitBinder<FNDITransformHandler, TNDIParamBinder<0, float, NDI_FUNC_BINDER(UNiagaraDataInterfaceSpline, SampleSplineUpVectorByUnitDistance)>>::Bind(this, BindingInfo, InstanceData, OutFunc);
 	}
 	else if (BindingInfo.Name == SampleSplineDirectionByUnitDistanceName)
 	{
 		check(BindingInfo.GetNumInputs() == 2 && BindingInfo.GetNumOutputs() == 3);
-		TNDIExplicitBinder<FNDITransformHandlerNoop, TNDIParamBinder<0, float, NDI_RAW_FUNC_BINDER(UNiagaraDataInterfaceSpline, SampleSplineDirectionByUnitDistance)>>::Bind(this, BindingInfo, InstanceData, OutFunc);
+		TNDIExplicitBinder<FNDITransformHandlerNoop, TNDIParamBinder<0, float, NDI_FUNC_BINDER(UNiagaraDataInterfaceSpline, SampleSplineDirectionByUnitDistance)>>::Bind(this, BindingInfo, InstanceData, OutFunc);
 	}
 	else if (BindingInfo.Name == SampleSplineDirectionByUnitDistanceWSName)
 	{
 		check(BindingInfo.GetNumInputs() == 2 && BindingInfo.GetNumOutputs() == 3);
-		TNDIExplicitBinder<FNDITransformHandler, TNDIParamBinder<0, float, NDI_RAW_FUNC_BINDER(UNiagaraDataInterfaceSpline, SampleSplineDirectionByUnitDistance)>>::Bind(this, BindingInfo, InstanceData, OutFunc);
+		TNDIExplicitBinder<FNDITransformHandler, TNDIParamBinder<0, float, NDI_FUNC_BINDER(UNiagaraDataInterfaceSpline, SampleSplineDirectionByUnitDistance)>>::Bind(this, BindingInfo, InstanceData, OutFunc);
 	}
 	else if (BindingInfo.Name == SampleSplineRightVectorByUnitDistanceName)
 	{
 		check(BindingInfo.GetNumInputs() == 2 && BindingInfo.GetNumOutputs() == 3);
-		TNDIExplicitBinder<FNDITransformHandlerNoop, TNDIParamBinder<0, float, NDI_RAW_FUNC_BINDER(UNiagaraDataInterfaceSpline, SampleSplineRightVectorByUnitDistance)>>::Bind(this, BindingInfo, InstanceData, OutFunc);
+		TNDIExplicitBinder<FNDITransformHandlerNoop, TNDIParamBinder<0, float, NDI_FUNC_BINDER(UNiagaraDataInterfaceSpline, SampleSplineRightVectorByUnitDistance)>>::Bind(this, BindingInfo, InstanceData, OutFunc);
 	}
 	else if (BindingInfo.Name == SampleSplineRightVectorByUnitDistanceWSName)
 	{
 		check(BindingInfo.GetNumInputs() == 2 && BindingInfo.GetNumOutputs() == 3);
-		TNDIExplicitBinder<FNDITransformHandler, TNDIParamBinder<0, float, NDI_RAW_FUNC_BINDER(UNiagaraDataInterfaceSpline, SampleSplineRightVectorByUnitDistance)>>::Bind(this, BindingInfo, InstanceData, OutFunc);
+		TNDIExplicitBinder<FNDITransformHandler, TNDIParamBinder<0, float, NDI_FUNC_BINDER(UNiagaraDataInterfaceSpline, SampleSplineRightVectorByUnitDistance)>>::Bind(this, BindingInfo, InstanceData, OutFunc);
 	}
 	else if (BindingInfo.Name == SampleSplineTangentByUnitDistanceName)
 	{
 		check(BindingInfo.GetNumInputs() == 2 && BindingInfo.GetNumOutputs() == 3);
-		TNDIExplicitBinder<FNDITransformHandlerNoop, TNDIParamBinder<0, float, NDI_RAW_FUNC_BINDER(UNiagaraDataInterfaceSpline, SampleSplineTangentByUnitDistance)>>::Bind(this, BindingInfo, InstanceData, OutFunc);
+		TNDIExplicitBinder<FNDITransformHandlerNoop, TNDIParamBinder<0, float, NDI_FUNC_BINDER(UNiagaraDataInterfaceSpline, SampleSplineTangentByUnitDistance)>>::Bind(this, BindingInfo, InstanceData, OutFunc);
 	}
 	else if (BindingInfo.Name == SampleSplineTangentByUnitDistanceWSName)
 	{
 		check(BindingInfo.GetNumInputs() == 2 && BindingInfo.GetNumOutputs() == 3);
-		TNDIExplicitBinder<FNDITransformHandler, TNDIParamBinder<0, float, NDI_RAW_FUNC_BINDER(UNiagaraDataInterfaceSpline, SampleSplineTangentByUnitDistance)>>::Bind(this, BindingInfo, InstanceData, OutFunc);
+		TNDIExplicitBinder<FNDITransformHandler, TNDIParamBinder<0, float, NDI_FUNC_BINDER(UNiagaraDataInterfaceSpline, SampleSplineTangentByUnitDistance)>>::Bind(this, BindingInfo, InstanceData, OutFunc);
 	}
 	else if (BindingInfo.Name == FindClosestUnitDistanceFromPositionWSName)
 	{
 		check(BindingInfo.GetNumInputs() == 4 && BindingInfo.GetNumOutputs() == 1);
-		TNDIParamBinder<0, float, TNDIParamBinder<1, float, TNDIParamBinder<2, float, NDI_RAW_FUNC_BINDER(UNiagaraDataInterfaceSpline, FindClosestUnitDistanceFromPositionWS)>>>::Bind(this, BindingInfo, InstanceData, OutFunc);
+		TNDIParamBinder<0, float, TNDIParamBinder<1, float, TNDIParamBinder<2, float, NDI_FUNC_BINDER(UNiagaraDataInterfaceSpline, FindClosestUnitDistanceFromPositionWS)>>>::Bind(this, BindingInfo, InstanceData, OutFunc);
 	}
 	else if (BindingInfo.Name == GetSplineLocalToWorldName)
 	{
@@ -393,10 +393,10 @@ void UNiagaraDataInterfaceSpline::SampleSplinePositionByUnitDistance(FVectorVMCo
 {
 	TransformHandlerType TransformHandler;
 	SplineSampleType SplineSampleParam(Context);
-	FUserPtrHandler<FNDISpline_InstanceData> InstData(Context);
-	FRegisterHandler<float> OutPosX(Context);
-	FRegisterHandler<float> OutPosY(Context);
-	FRegisterHandler<float> OutPosZ(Context);
+	VectorVM::FUserPtrHandler<FNDISpline_InstanceData> InstData(Context);
+	VectorVM::FExternalFuncRegisterHandler<float> OutPosX(Context);
+	VectorVM::FExternalFuncRegisterHandler<float> OutPosY(Context);
+	VectorVM::FExternalFuncRegisterHandler<float> OutPosZ(Context);
 
 	if (InstData->Component)
 	{
@@ -442,10 +442,10 @@ void UNiagaraDataInterfaceSpline::SampleSplineUpVectorByUnitDistance(FVectorVMCo
 {
 	TransformHandlerType TransformHandler;
 	SplineSampleType SplineSampleParam(Context);
-	FUserPtrHandler<FNDISpline_InstanceData> InstData(Context);
-	FRegisterHandler<float> OutPosX(Context);
-	FRegisterHandler<float> OutPosY(Context);
-	FRegisterHandler<float> OutPosZ(Context);
+	VectorVM::FUserPtrHandler<FNDISpline_InstanceData> InstData(Context);
+	VectorVM::FExternalFuncRegisterHandler<float> OutPosX(Context);
+	VectorVM::FExternalFuncRegisterHandler<float> OutPosY(Context);
+	VectorVM::FExternalFuncRegisterHandler<float> OutPosZ(Context);
 
 	if (InstData->Component)
 	{
@@ -491,10 +491,10 @@ void UNiagaraDataInterfaceSpline::SampleSplineRightVectorByUnitDistance(FVectorV
 {
 	TransformHandlerType TransformHandler;
 	SplineSampleType SplineSampleParam(Context);
-	FUserPtrHandler<FNDISpline_InstanceData> InstData(Context);
-	FRegisterHandler<float> OutPosX(Context);
-	FRegisterHandler<float> OutPosY(Context);
-	FRegisterHandler<float> OutPosZ(Context);
+	VectorVM::FUserPtrHandler<FNDISpline_InstanceData> InstData(Context);
+	VectorVM::FExternalFuncRegisterHandler<float> OutPosX(Context);
+	VectorVM::FExternalFuncRegisterHandler<float> OutPosY(Context);
+	VectorVM::FExternalFuncRegisterHandler<float> OutPosZ(Context);
 	
 	
 	if (InstData->Component)
@@ -538,10 +538,10 @@ void UNiagaraDataInterfaceSpline::SampleSplineTangentByUnitDistance(FVectorVMCon
 {
 	TransformHandlerType TransformHandler;
 	SplineSampleType SplineSampleParam(Context);
-	FUserPtrHandler<FNDISpline_InstanceData> InstData(Context);
-	FRegisterHandler<float> OutPosX(Context);
-	FRegisterHandler<float> OutPosY(Context);
-	FRegisterHandler<float> OutPosZ(Context);
+	VectorVM::FUserPtrHandler<FNDISpline_InstanceData> InstData(Context);
+	VectorVM::FExternalFuncRegisterHandler<float> OutPosX(Context);
+	VectorVM::FExternalFuncRegisterHandler<float> OutPosY(Context);
+	VectorVM::FExternalFuncRegisterHandler<float> OutPosZ(Context);
 
 	if (InstData->Component)
 	{
@@ -586,10 +586,10 @@ void UNiagaraDataInterfaceSpline::SampleSplineDirectionByUnitDistance(FVectorVMC
 {
 	TransformHandlerType TransformHandler;
 	SplineSampleType SplineSampleParam(Context);
-	FUserPtrHandler<FNDISpline_InstanceData> InstData(Context);
-	FRegisterHandler<float> OutPosX(Context);
-	FRegisterHandler<float> OutPosY(Context);
-	FRegisterHandler<float> OutPosZ(Context);
+	VectorVM::FUserPtrHandler<FNDISpline_InstanceData> InstData(Context);
+	VectorVM::FExternalFuncRegisterHandler<float> OutPosX(Context);
+	VectorVM::FExternalFuncRegisterHandler<float> OutPosY(Context);
+	VectorVM::FExternalFuncRegisterHandler<float> OutPosZ(Context);
 
 	if (InstData->Component)
 	{
@@ -631,22 +631,22 @@ void UNiagaraDataInterfaceSpline::SampleSplineDirectionByUnitDistance(FVectorVMC
 
 void UNiagaraDataInterfaceSpline::WriteTransform(const FMatrix& ToWrite, FVectorVMContext& Context)
 {
-	FRegisterHandler<float> Out00(Context);
-	FRegisterHandler<float> Out01(Context);
-	FRegisterHandler<float> Out02(Context);
-	FRegisterHandler<float> Out03(Context);
-	FRegisterHandler<float> Out04(Context);
-	FRegisterHandler<float> Out05(Context);
-	FRegisterHandler<float> Out06(Context);
-	FRegisterHandler<float> Out07(Context);
-	FRegisterHandler<float> Out08(Context);
-	FRegisterHandler<float> Out09(Context);
-	FRegisterHandler<float> Out10(Context);
-	FRegisterHandler<float> Out11(Context);
-	FRegisterHandler<float> Out12(Context);
-	FRegisterHandler<float> Out13(Context);
-	FRegisterHandler<float> Out14(Context);
-	FRegisterHandler<float> Out15(Context);
+	VectorVM::FExternalFuncRegisterHandler<float> Out00(Context);
+	VectorVM::FExternalFuncRegisterHandler<float> Out01(Context);
+	VectorVM::FExternalFuncRegisterHandler<float> Out02(Context);
+	VectorVM::FExternalFuncRegisterHandler<float> Out03(Context);
+	VectorVM::FExternalFuncRegisterHandler<float> Out04(Context);
+	VectorVM::FExternalFuncRegisterHandler<float> Out05(Context);
+	VectorVM::FExternalFuncRegisterHandler<float> Out06(Context);
+	VectorVM::FExternalFuncRegisterHandler<float> Out07(Context);
+	VectorVM::FExternalFuncRegisterHandler<float> Out08(Context);
+	VectorVM::FExternalFuncRegisterHandler<float> Out09(Context);
+	VectorVM::FExternalFuncRegisterHandler<float> Out10(Context);
+	VectorVM::FExternalFuncRegisterHandler<float> Out11(Context);
+	VectorVM::FExternalFuncRegisterHandler<float> Out12(Context);
+	VectorVM::FExternalFuncRegisterHandler<float> Out13(Context);
+	VectorVM::FExternalFuncRegisterHandler<float> Out14(Context);
+	VectorVM::FExternalFuncRegisterHandler<float> Out15(Context);
 
 	for (int32 i = 0; i < Context.NumInstances; ++i)
 	{
@@ -675,8 +675,8 @@ void UNiagaraDataInterfaceSpline::FindClosestUnitDistanceFromPositionWS(FVectorV
 	PosXType PosXParam(Context);
 	PosYType PosYParam(Context);
 	PosZType PosZParam(Context);
-	FUserPtrHandler<FNDISpline_InstanceData> InstData(Context);
-	FRegisterHandler<float> OutUnitDistance(Context);
+	VectorVM::FUserPtrHandler<FNDISpline_InstanceData> InstData(Context);
+	VectorVM::FExternalFuncRegisterHandler<float> OutUnitDistance(Context);
 
 	if (InstData->Component)
 	{
@@ -725,13 +725,13 @@ void UNiagaraDataInterfaceSpline::FindClosestUnitDistanceFromPositionWS(FVectorV
 
 void UNiagaraDataInterfaceSpline::GetLocalToWorld(FVectorVMContext& Context)
 {
-	FUserPtrHandler<FNDISpline_InstanceData> InstData(Context);
+	VectorVM::FUserPtrHandler<FNDISpline_InstanceData> InstData(Context);
 	WriteTransform(InstData->Transform, Context);
 }
 
 void UNiagaraDataInterfaceSpline::GetLocalToWorldInverseTransposed(FVectorVMContext& Context)
 {
-	FUserPtrHandler<FNDISpline_InstanceData> InstData(Context);
+	VectorVM::FUserPtrHandler<FNDISpline_InstanceData> InstData(Context);
 	WriteTransform(InstData->TransformInverseTransposed, Context);
 }
 
diff --git a/Engine/Plugins/FX/Niagara/Source/Niagara/Private/NiagaraDataInterfaceStaticMesh.cpp b/Engine/Plugins/FX/Niagara/Source/Niagara/Private/NiagaraDataInterfaceStaticMesh.cpp
--- a/Engine/Plugins/FX/Niagara/Source/Niagara/Private/NiagaraDataInterfaceStaticMesh.cpp
+++ b/Engine/Plugins/FX/Niagara/Source/Niagara/Private/NiagaraDataInterfaceStaticMesh.cpp
@@ -695,18 +695,18 @@ struct TTypedMeshAccessorBinder
 };
 
 //Final binders for all static mesh interface functions.
-DEFINE_NDI_RAW_FUNC_BINDER(UNiagaraDataInterfaceStaticMesh, RandomSection);
-DEFINE_NDI_RAW_FUNC_BINDER(UNiagaraDataInterfaceStaticMesh, RandomTriCoord);
-DEFINE_NDI_RAW_FUNC_BINDER(UNiagaraDataInterfaceStaticMesh, RandomTriCoordVertexColorFiltered);
-DEFINE_NDI_RAW_FUNC_BINDER(UNiagaraDataInterfaceStaticMesh, RandomTriCoordOnSection);
-DEFINE_NDI_RAW_FUNC_BINDER(UNiagaraDataInterfaceStaticMesh, GetTriCoordPosition);
-DEFINE_NDI_RAW_FUNC_BINDER(UNiagaraDataInterfaceStaticMesh, GetTriCoordNormal);
-DEFINE_NDI_RAW_FUNC_BINDER(UNiagaraDataInterfaceStaticMesh, GetTriCoordTangents);
-DEFINE_NDI_RAW_FUNC_BINDER(UNiagaraDataInterfaceStaticMesh, GetTriCoordColor);
-DEFINE_NDI_RAW_FUNC_BINDER(UNiagaraDataInterfaceStaticMesh, GetTriCoordUV);
-DEFINE_NDI_RAW_FUNC_BINDER(UNiagaraDataInterfaceStaticMesh, GetTriCoordPositionAndVelocity);
-
-DEFINE_NDI_RAW_FUNC_BINDER(UNiagaraDataInterfaceStaticMesh, GetVertexPosition);
+DEFINE_NDI_FUNC_BINDER(UNiagaraDataInterfaceStaticMesh, RandomSection);
+DEFINE_NDI_FUNC_BINDER(UNiagaraDataInterfaceStaticMesh, RandomTriCoord);
+DEFINE_NDI_DIRECT_FUNC_BINDER(UNiagaraDataInterfaceStaticMesh, RandomTriCoordVertexColorFiltered);
+DEFINE_NDI_FUNC_BINDER(UNiagaraDataInterfaceStaticMesh, RandomTriCoordOnSection);
+DEFINE_NDI_FUNC_BINDER(UNiagaraDataInterfaceStaticMesh, GetTriCoordPosition);
+DEFINE_NDI_FUNC_BINDER(UNiagaraDataInterfaceStaticMesh, GetTriCoordNormal);
+DEFINE_NDI_FUNC_BINDER(UNiagaraDataInterfaceStaticMesh, GetTriCoordTangents);
+DEFINE_NDI_DIRECT_FUNC_BINDER(UNiagaraDataInterfaceStaticMesh, GetTriCoordColor);
+DEFINE_NDI_FUNC_BINDER(UNiagaraDataInterfaceStaticMesh, GetTriCoordUV);
+DEFINE_NDI_DIRECT_FUNC_BINDER(UNiagaraDataInterfaceStaticMesh, GetTriCoordPositionAndVelocity);
+
+DEFINE_NDI_FUNC_BINDER(UNiagaraDataInterfaceStaticMesh, GetVertexPosition);
 
 void UNiagaraDataInterfaceStaticMesh::GetVMExternalFunction(const FVMExternalFunctionBindingInfo& BindingInfo, void* InstanceData, FVMExternalFunction &OutFunc)
 {
@@ -720,79 +720,79 @@ void UNiagaraDataInterfaceStaticMesh::GetVMExternalFunction(const FVMExternalFun
 	if (BindingInfo.Name == StaticMeshHelpers::RandomSectionName)
 	{
 		check(BindingInfo.GetNumInputs() == 1 && BindingInfo.GetNumOutputs() == 1);
-		TUsesAreaWeightingBinder<NDI_RAW_FUNC_BINDER(UNiagaraDataInterfaceStaticMesh, RandomSection)>::Bind(this, BindingInfo, InstanceData, OutFunc);
+		TUsesAreaWeightingBinder<NDI_FUNC_BINDER(UNiagaraDataInterfaceStaticMesh, RandomSection)>::Bind(this, BindingInfo, InstanceData, OutFunc);
 	}
 	else if (BindingInfo.Name == StaticMeshHelpers::RandomTriCoordName)
 	{
 		check(BindingInfo.GetNumInputs() == 1 && BindingInfo.GetNumOutputs() == 4);
-		TUsesAreaWeightingBinder<NDI_RAW_FUNC_BINDER(UNiagaraDataInterfaceStaticMesh, RandomTriCoord)>::Bind(this, BindingInfo, InstanceData, OutFunc);
+		TUsesAreaWeightingBinder<NDI_FUNC_BINDER(UNiagaraDataInterfaceStaticMesh, RandomTriCoord)>::Bind(this, BindingInfo, InstanceData, OutFunc);
 	}
 	//TODO: Vertex color filtering needs more work.
 	else if (BindingInfo.Name == StaticMeshHelpers::RandomTriCoordVCFilteredName)
 	{
 		InstData->InitVertexColorFiltering();
 		check(BindingInfo.GetNumInputs() == 3 && BindingInfo.GetNumOutputs() == 4);
-		TNDIParamBinder<0, float, TNDIParamBinder<1, float, NDI_RAW_FUNC_BINDER(UNiagaraDataInterfaceStaticMesh, RandomTriCoordVertexColorFiltered)>>::Bind(this, BindingInfo, InstanceData, OutFunc);
+		NDI_FUNC_BINDER(UNiagaraDataInterfaceStaticMesh, RandomTriCoordVertexColorFiltered)::Bind(this, OutFunc);
 	}	
 	else if (BindingInfo.Name == StaticMeshHelpers::RandomTriCoordOnSectionName)
 	{
 		check(BindingInfo.GetNumInputs() == 2 && BindingInfo.GetNumOutputs() == 4);
-		TUsesAreaWeightingBinder<TNDIParamBinder<0, int32, NDI_RAW_FUNC_BINDER(UNiagaraDataInterfaceStaticMesh, RandomTriCoordOnSection)>>::Bind(this, BindingInfo, InstanceData, OutFunc);
+		TUsesAreaWeightingBinder<NDI_FUNC_BINDER(UNiagaraDataInterfaceStaticMesh, RandomTriCoordOnSection)>::Bind(this, BindingInfo, InstanceData, OutFunc);
 	}
 	else if (BindingInfo.Name == StaticMeshHelpers::GetTriPositionName)
 	{
 		check(BindingInfo.GetNumInputs() == 5 && BindingInfo.GetNumOutputs() == 3);
 		bNeedsVertexPositions = true;
-		TNDIExplicitBinder<FNDITransformHandlerNoop, TNDIParamBinder<0, int32, TNDIParamBinder<1, float, TNDIParamBinder<2, float, TNDIParamBinder<3, float, NDI_RAW_FUNC_BINDER(UNiagaraDataInterfaceStaticMesh, GetTriCoordPosition)>>>>>::Bind(this, BindingInfo, InstanceData, OutFunc);
+		TNDIExplicitBinder<FNDITransformHandlerNoop, NDI_FUNC_BINDER(UNiagaraDataInterfaceStaticMesh, GetTriCoordPosition)>::Bind(this, BindingInfo, InstanceData, OutFunc);
 	}
 	else if (BindingInfo.Name == StaticMeshHelpers::GetTriPositionWSName)
 	{
 		check(BindingInfo.GetNumInputs() == 5 && BindingInfo.GetNumOutputs() == 3);
 		bNeedsVertexPositions = true;
-		TNDIExplicitBinder<FNDITransformHandler, TNDIParamBinder<0, int32, TNDIParamBinder<1, float, TNDIParamBinder<2, float, TNDIParamBinder<3, float, NDI_RAW_FUNC_BINDER(UNiagaraDataInterfaceStaticMesh, GetTriCoordPosition)>>>>>::Bind(this, BindingInfo, InstanceData, OutFunc);
+		TNDIExplicitBinder<FNDITransformHandler, NDI_FUNC_BINDER(UNiagaraDataInterfaceStaticMesh, GetTriCoordPosition)>::Bind(this, BindingInfo, InstanceData, OutFunc);
 	}
 	else if (BindingInfo.Name == StaticMeshHelpers::GetTriNormalName)
 	{
 		check(BindingInfo.GetNumInputs() == 6 && BindingInfo.GetNumOutputs() == 3);
 		bNeedsVertMain = true;
-		TNDIExplicitBinder<FNDITransformHandlerNoop, TNDIParamBinder<0, int32, TNDIParamBinder<1, float, TNDIParamBinder<2, float, TNDIParamBinder<3, float, NDI_RAW_FUNC_BINDER(UNiagaraDataInterfaceStaticMesh, GetTriCoordNormal)>>>>>::Bind(this, BindingInfo, InstanceData, OutFunc);
+		TNDIExplicitBinder<FNDITransformHandlerNoop, NDI_FUNC_BINDER(UNiagaraDataInterfaceStaticMesh, GetTriCoordNormal)>::Bind(this, BindingInfo, InstanceData, OutFunc);
 	}
 	else if (BindingInfo.Name == StaticMeshHelpers::GetTriNormalWSName)
 	{
 		check(BindingInfo.GetNumInputs() == 5 && BindingInfo.GetNumOutputs() == 3);
 		bNeedsVertMain = true;
-		TNDIExplicitBinder<FNDITransformHandler, TNDIParamBinder<0, int32, TNDIParamBinder<1, float, TNDIParamBinder<2, float, TNDIParamBinder<3, float, NDI_RAW_FUNC_BINDER(UNiagaraDataInterfaceStaticMesh, GetTriCoordNormal)>>>>>::Bind(this, BindingInfo, InstanceData, OutFunc);
+		TNDIExplicitBinder<FNDITransformHandler, NDI_FUNC_BINDER(UNiagaraDataInterfaceStaticMesh, GetTriCoordNormal)>::Bind(this, BindingInfo, InstanceData, OutFunc);
 	}
 	else if (BindingInfo.Name == StaticMeshHelpers::GetTriTangentsName)
 	{
 		check(BindingInfo.GetNumInputs() == 5 && BindingInfo.GetNumOutputs() == 9);
 		bNeedsVertMain = true;
-		TTypedMeshAccessorBinder<TNDIExplicitBinder<FNDITransformHandlerNoop, TNDIParamBinder<0, int32, TNDIParamBinder<1, float, TNDIParamBinder<2, float, TNDIParamBinder<3, float, NDI_RAW_FUNC_BINDER(UNiagaraDataInterfaceStaticMesh, GetTriCoordTangents)>>>>>>::Bind(this, BindingInfo, InstanceData, OutFunc);
+		TTypedMeshAccessorBinder<TNDIExplicitBinder<FNDITransformHandlerNoop, NDI_FUNC_BINDER(UNiagaraDataInterfaceStaticMesh, GetTriCoordTangents)>>::Bind(this, BindingInfo, InstanceData, OutFunc);
 	}
 	else if (BindingInfo.Name == StaticMeshHelpers::GetTriTangentsWSName)
 	{
 		check(BindingInfo.GetNumInputs() == 5 && BindingInfo.GetNumOutputs() == 9);
 		bNeedsVertMain = true;
-		TTypedMeshAccessorBinder<TNDIExplicitBinder<FNDITransformHandler, TNDIParamBinder<0, int32, TNDIParamBinder<1, float, TNDIParamBinder<2, float, TNDIParamBinder<3, float, NDI_RAW_FUNC_BINDER(UNiagaraDataInterfaceStaticMesh, GetTriCoordTangents)>>>>>>::Bind(this, BindingInfo, InstanceData, OutFunc);
+		TTypedMeshAccessorBinder<TNDIExplicitBinder<FNDITransformHandler, NDI_FUNC_BINDER(UNiagaraDataInterfaceStaticMesh, GetTriCoordTangents)>>::Bind(this, BindingInfo, InstanceData, OutFunc);
 	}
 	else if (BindingInfo.Name == StaticMeshHelpers::GetTriColorName)
 	{
 		check(BindingInfo.GetNumInputs() == 5 && BindingInfo.GetNumOutputs() == 4);
 		bNeedsVertexColors = true;
-		TNDIParamBinder<0, int32, TNDIParamBinder<1, float, TNDIParamBinder<2, float, TNDIParamBinder<3, float, NDI_RAW_FUNC_BINDER(UNiagaraDataInterfaceStaticMesh, GetTriCoordColor)>>>>::Bind(this, BindingInfo, InstanceData, OutFunc);
+		NDI_FUNC_BINDER(UNiagaraDataInterfaceStaticMesh, GetTriCoordColor)::Bind(this, OutFunc);
 	}
 	else if (BindingInfo.Name == StaticMeshHelpers::GetTriUVName)
 	{
 		check(BindingInfo.GetNumInputs() == 6 && BindingInfo.GetNumOutputs() == 2);
 		bNeedsVertMain = true;
-		TTypedMeshAccessorBinder<TNDIParamBinder<0, int32, TNDIParamBinder<1, float, TNDIParamBinder<2, float, TNDIParamBinder<3, float, TNDIParamBinder<4, int32, NDI_RAW_FUNC_BINDER(UNiagaraDataInterfaceStaticMesh, GetTriCoordUV)>>>>>>::Bind(this, BindingInfo, InstanceData, OutFunc);
+		TTypedMeshAccessorBinder<NDI_FUNC_BINDER(UNiagaraDataInterfaceStaticMesh, GetTriCoordUV)>::Bind(this, BindingInfo, InstanceData, OutFunc);
 	}
 	else if (BindingInfo.Name == StaticMeshHelpers::GetTriPositionAndVelocityName)
 	{
 		check(BindingInfo.GetNumInputs() == 5 && BindingInfo.GetNumOutputs() == 6);
 		bNeedsVertMain = true;
-		bNeedsVertexPositions = true;
-		TNDIParamBinder<0, int32, TNDIParamBinder<1, float, TNDIParamBinder<2, float, TNDIParamBinder<3, float, NDI_RAW_FUNC_BINDER(UNiagaraDataInterfaceStaticMesh, GetTriCoordPositionAndVelocity)>>>>::Bind(this, BindingInfo, InstanceData, OutFunc);
+		bNeedsVertexPositions = true;  
+		NDI_FUNC_BINDER(UNiagaraDataInterfaceStaticMesh, GetTriCoordPositionAndVelocity)::Bind(this, OutFunc);
 	}
 	else if (BindingInfo.Name == StaticMeshHelpers::GetMeshLocalToWorldName)
 	{
@@ -813,13 +813,13 @@ void UNiagaraDataInterfaceStaticMesh::GetVMExternalFunction(const FVMExternalFun
 	{
 		check(BindingInfo.GetNumInputs() == 2 && BindingInfo.GetNumOutputs() == 3);
 		bNeedsVertexPositions = true;
-		TNDIExplicitBinder<FNDITransformHandlerNoop, TNDIParamBinder<0, int32, NDI_RAW_FUNC_BINDER(UNiagaraDataInterfaceStaticMesh, GetVertexPosition)>>::Bind(this, BindingInfo, InstanceData, OutFunc);
+		TNDIExplicitBinder<FNDITransformHandlerNoop, NDI_FUNC_BINDER(UNiagaraDataInterfaceStaticMesh, GetVertexPosition)>::Bind(this, BindingInfo, InstanceData, OutFunc);
 	}
 	else if (BindingInfo.Name == StaticMeshHelpers::GetVertexPositionWSName)
 	{
 		check(BindingInfo.GetNumInputs() == 2 && BindingInfo.GetNumOutputs() == 3);
 		bNeedsVertexPositions = true;
-		TNDIExplicitBinder<FNDITransformHandler, TNDIParamBinder<0, int32, NDI_RAW_FUNC_BINDER(UNiagaraDataInterfaceStaticMesh, GetVertexPosition)>>::Bind(this, BindingInfo, InstanceData, OutFunc);
+		TNDIExplicitBinder<FNDITransformHandler, NDI_FUNC_BINDER(UNiagaraDataInterfaceStaticMesh, GetVertexPosition)>::Bind(this, BindingInfo, InstanceData, OutFunc);
 	}
 
 	if (bNeedsVertexPositions && !InstData->MeshHasPositions())
@@ -944,8 +944,8 @@ FORCEINLINE int32 UNiagaraDataInterfaceStaticMesh::RandomSection<TIntegralConsta
 template<typename TAreaWeighted>
 void UNiagaraDataInterfaceStaticMesh::RandomSection(FVectorVMContext& Context)
 {
-	FUserPtrHandler<FNDIStaticMesh_InstanceData> InstData(Context);
-	FRegisterHandler<int32> OutSection(Context);
+	VectorVM::FUserPtrHandler<FNDIStaticMesh_InstanceData> InstData(Context);
+	VectorVM::FExternalFuncRegisterHandler<int32> OutSection(Context);
 
 	FStaticMeshLODResources& Res = InstData->Mesh->RenderData->LODResources[0];
 	for (int32 i = 0; i < Context.NumInstances; ++i)
@@ -996,12 +996,12 @@ FORCEINLINE int32 UNiagaraDataInterfaceStaticMesh::RandomTriIndex<TIntegralConst
 template<typename TAreaWeighted>
 void UNiagaraDataInterfaceStaticMesh::RandomTriCoord(FVectorVMContext& Context)
 {
-	FUserPtrHandler<FNDIStaticMesh_InstanceData> InstData(Context);
+	VectorVM::FUserPtrHandler<FNDIStaticMesh_InstanceData> InstData(Context);
 
-	FRegisterHandler<int32> OutTri(Context);
-	FRegisterHandler<float> OutBaryX(Context);
-	FRegisterHandler<float> OutBaryY(Context);
-	FRegisterHandler<float> OutBaryZ(Context);
+	VectorVM::FExternalFuncRegisterHandler<int32> OutTri(Context);
+	VectorVM::FExternalFuncRegisterHandler<float> OutBaryX(Context);
+	VectorVM::FExternalFuncRegisterHandler<float> OutBaryY(Context);
+	VectorVM::FExternalFuncRegisterHandler<float> OutBaryZ(Context);
 
 	FStaticMeshLODResources& Res = InstData->Mesh->RenderData->LODResources[0];
 	FIndexArrayView Indices = Res.IndexBuffer.GetArrayView();
@@ -1020,17 +1020,16 @@ void UNiagaraDataInterfaceStaticMesh::RandomTriCoord(FVectorVMContext& Context)
 	}
 }
 
-template<typename InputType0, typename InputType1>
 void UNiagaraDataInterfaceStaticMesh::RandomTriCoordVertexColorFiltered(FVectorVMContext& Context)
 {
-	InputType0 MinValue(Context);
-	InputType1 RangeValue(Context);
-	FUserPtrHandler<FNDIStaticMesh_InstanceData> InstData(Context);
-
-	FRegisterHandler<int32> OutTri(Context);
-	FRegisterHandler<float> OutBaryX(Context);
-	FRegisterHandler<float> OutBaryY(Context);
-	FRegisterHandler<float> OutBaryZ(Context);
+	VectorVM::FExternalFuncRegisterHandler<int32> MinValue(Context);
+	VectorVM::FExternalFuncRegisterHandler<int32> RangeValue(Context);
+	VectorVM::FUserPtrHandler<FNDIStaticMesh_InstanceData> InstData(Context);
+
+	VectorVM::FExternalFuncRegisterHandler<int32> OutTri(Context);
+	VectorVM::FExternalFuncRegisterHandler<float> OutBaryX(Context);
+	VectorVM::FExternalFuncRegisterHandler<float> OutBaryY(Context);
+	VectorVM::FExternalFuncRegisterHandler<float> OutBaryZ(Context);
 	
 	FDynamicVertexColorFilterData* VCFData = InstData->DynamicVertexColorSampler.Get();
 
@@ -1099,16 +1098,16 @@ FORCEINLINE int32 UNiagaraDataInterfaceStaticMesh::RandomTriIndexOnSection<TInte
 	return Sec.FirstIndex + Tri * 3;
 }
 
-template<typename TAreaWeighted, typename SectionIdxType>
+template<typename TAreaWeighted>
 void UNiagaraDataInterfaceStaticMesh::RandomTriCoordOnSection(FVectorVMContext& Context)
 {
-	SectionIdxType SectionIdxParam(Context);
-	FUserPtrHandler<FNDIStaticMesh_InstanceData> InstData(Context);
+	VectorVM::FExternalFuncInputHandler<int32> SectionIdxParam(Context);
+	VectorVM::FUserPtrHandler<FNDIStaticMesh_InstanceData> InstData(Context);
 
-	FRegisterHandler<int32> OutTri(Context);
-	FRegisterHandler<float> OutBaryX(Context);
-	FRegisterHandler<float> OutBaryY(Context);
-	FRegisterHandler<float> OutBaryZ(Context);
+	VectorVM::FExternalFuncRegisterHandler<int32> OutTri(Context);
+	VectorVM::FExternalFuncRegisterHandler<float> OutBaryX(Context);
+	VectorVM::FExternalFuncRegisterHandler<float> OutBaryY(Context);
+	VectorVM::FExternalFuncRegisterHandler<float> OutBaryZ(Context);
 
 	FStaticMeshLODResources& Res = InstData->Mesh->RenderData->LODResources[0];
 	FIndexArrayView Indices = Res.IndexBuffer.GetArrayView();
@@ -1129,19 +1128,19 @@ void UNiagaraDataInterfaceStaticMesh::RandomTriCoordOnSection(FVectorVMContext&
 	}
 }
 
-template<typename TransformHandlerType, typename TriType, typename BaryXType, typename BaryYType, typename BaryZType>
+template<typename TransformHandlerType>
 void UNiagaraDataInterfaceStaticMesh::GetTriCoordPosition(FVectorVMContext& Context)
 {
 	TransformHandlerType TransformHandler;
-	TriType TriParam(Context);
-	BaryXType BaryXParam(Context);
-	BaryYType BaryYParam(Context);
-	BaryZType BaryZParam(Context);
-	FUserPtrHandler<FNDIStaticMesh_InstanceData> InstData(Context);
+	VectorVM::FExternalFuncInputHandler<int32> TriParam(Context);
+	VectorVM::FExternalFuncInputHandler<float> BaryXParam(Context);
+	VectorVM::FExternalFuncInputHandler<float> BaryYParam(Context);
+	VectorVM::FExternalFuncInputHandler<float> BaryZParam(Context);
+	VectorVM::FUserPtrHandler<FNDIStaticMesh_InstanceData> InstData(Context);
 
-	FRegisterHandler<float> OutPosX(Context);
-	FRegisterHandler<float> OutPosY(Context);
-	FRegisterHandler<float> OutPosZ(Context);
+	VectorVM::FExternalFuncRegisterHandler<float> OutPosX(Context);
+	VectorVM::FExternalFuncRegisterHandler<float> OutPosY(Context);
+	VectorVM::FExternalFuncRegisterHandler<float> OutPosZ(Context);
 
 	FStaticMeshLODResources& Res = InstData->Mesh->RenderData->LODResources[0];
 	const FIndexArrayView& Indices = Res.IndexBuffer.GetArrayView();
@@ -1171,20 +1170,20 @@ void UNiagaraDataInterfaceStaticMesh::GetTriCoordPosition(FVectorVMContext& Cont
 	}
 }
 
-template<typename TransformHandlerType, typename TriType, typename BaryXType, typename BaryYType, typename BaryZType>
+template<typename TransformHandlerType>
 void UNiagaraDataInterfaceStaticMesh::GetTriCoordNormal(FVectorVMContext& Context)
 {
 	TransformHandlerType TransformHandler;
 
-	TriType TriParam(Context);
-	BaryXType BaryXParam(Context);
-	BaryYType BaryYParam(Context);
-	BaryZType BaryZParam(Context);
-	FUserPtrHandler<FNDIStaticMesh_InstanceData> InstData(Context);
+	VectorVM::FExternalFuncInputHandler<int32> TriParam(Context);
+	VectorVM::FExternalFuncInputHandler<float> BaryXParam(Context);
+	VectorVM::FExternalFuncInputHandler<float> BaryYParam(Context);
+	VectorVM::FExternalFuncInputHandler<float> BaryZParam(Context);
+	VectorVM::FUserPtrHandler<FNDIStaticMesh_InstanceData> InstData(Context);
 
-	FRegisterHandler<float> OutNormX(Context);
-	FRegisterHandler<float> OutNormY(Context);
-	FRegisterHandler<float> OutNormZ(Context);
+	VectorVM::FExternalFuncRegisterHandler<float> OutNormX(Context);
+	VectorVM::FExternalFuncRegisterHandler<float> OutNormY(Context);
+	VectorVM::FExternalFuncRegisterHandler<float> OutNormZ(Context);
 
 	FStaticMeshLODResources& Res = InstData->Mesh->RenderData->LODResources[0];
 	const FIndexArrayView& Indices = Res.IndexBuffer.GetArrayView();
@@ -1213,30 +1212,30 @@ void UNiagaraDataInterfaceStaticMesh::GetTriCoordNormal(FVectorVMContext& Contex
 	}
 }
 
-template<typename VertexAccessorType, typename TransformHandlerType, typename TriType, typename BaryXType, typename BaryYType, typename BaryZType>
+template<typename VertexAccessorType, typename TransformHandlerType>
 void UNiagaraDataInterfaceStaticMesh::GetTriCoordTangents(FVectorVMContext& Context)
 {
 	TransformHandlerType TransformHandler;	
 
-	TriType TriParam(Context);
-	BaryXType BaryXParam(Context);
-	BaryYType BaryYParam(Context);
-	BaryZType BaryZParam(Context);
-	FUserPtrHandler<FNDIStaticMesh_InstanceData> InstData(Context);
+	VectorVM::FExternalFuncInputHandler<int32> TriParam(Context);
+	VectorVM::FExternalFuncInputHandler<float> BaryXParam(Context);
+	VectorVM::FExternalFuncInputHandler<float> BaryYParam(Context);
+	VectorVM::FExternalFuncInputHandler<float> BaryZParam(Context);
+	VectorVM::FUserPtrHandler<FNDIStaticMesh_InstanceData> InstData(Context);
 
 	FStaticMeshLODResources& Res = InstData->Mesh->RenderData->LODResources[0];
 	const FIndexArrayView& Indices = Res.IndexBuffer.GetArrayView();
 	const VertexAccessorType Verts(Res.VertexBuffers.StaticMeshVertexBuffer);
 
-	FRegisterHandler<float> OutTangentX(Context);
-	FRegisterHandler<float> OutTangentY(Context);
-	FRegisterHandler<float> OutTangentZ(Context);
-	FRegisterHandler<float> OutBinormX(Context);
-	FRegisterHandler<float> OutBinormY(Context);
-	FRegisterHandler<float> OutBinormZ(Context);
-	FRegisterHandler<float> OutNormX(Context);
-	FRegisterHandler<float> OutNormY(Context);
-	FRegisterHandler<float> OutNormZ(Context);
+	VectorVM::FExternalFuncRegisterHandler<float> OutTangentX(Context);
+	VectorVM::FExternalFuncRegisterHandler<float> OutTangentY(Context);
+	VectorVM::FExternalFuncRegisterHandler<float> OutTangentZ(Context);
+	VectorVM::FExternalFuncRegisterHandler<float> OutBinormX(Context);
+	VectorVM::FExternalFuncRegisterHandler<float> OutBinormY(Context);
+	VectorVM::FExternalFuncRegisterHandler<float> OutBinormZ(Context);
+	VectorVM::FExternalFuncRegisterHandler<float> OutNormX(Context);
+	VectorVM::FExternalFuncRegisterHandler<float> OutNormY(Context);
+	VectorVM::FExternalFuncRegisterHandler<float> OutNormZ(Context);
 
 	for (int32 i = 0; i < Context.NumInstances; ++i)
 	{
@@ -1276,19 +1275,18 @@ void UNiagaraDataInterfaceStaticMesh::GetTriCoordTangents(FVectorVMContext& Cont
 	}
 }
 
-template<typename TriType, typename BaryXType, typename BaryYType, typename BaryZType>
 void UNiagaraDataInterfaceStaticMesh::GetTriCoordColor(FVectorVMContext& Context)
 {
-	TriType TriParam(Context);
-	BaryXType BaryXParam(Context);
-	BaryYType BaryYParam(Context);
-	BaryZType BaryZParam(Context);
-	FUserPtrHandler<FNDIStaticMesh_InstanceData> InstData(Context);
+	VectorVM::FExternalFuncInputHandler<int32> TriParam(Context);
+	VectorVM::FExternalFuncInputHandler<float> BaryXParam(Context);
+	VectorVM::FExternalFuncInputHandler<float> BaryYParam(Context);
+	VectorVM::FExternalFuncInputHandler<float> BaryZParam(Context);
+	VectorVM::FUserPtrHandler<FNDIStaticMesh_InstanceData> InstData(Context);
 
-	FRegisterHandler<float> OutColorR(Context);
-	FRegisterHandler<float> OutColorG(Context);
-	FRegisterHandler<float> OutColorB(Context);
-	FRegisterHandler<float> OutColorA(Context);
+	VectorVM::FExternalFuncRegisterHandler<float> OutColorR(Context);
+	VectorVM::FExternalFuncRegisterHandler<float> OutColorG(Context);
+	VectorVM::FExternalFuncRegisterHandler<float> OutColorB(Context);
+	VectorVM::FExternalFuncRegisterHandler<float> OutColorA(Context);
 
 	FStaticMeshLODResources& Res = InstData->Mesh->RenderData->LODResources[0];
 	const FIndexArrayView& Indices = Res.IndexBuffer.GetArrayView();
@@ -1342,18 +1340,18 @@ void UNiagaraDataInterfaceStaticMesh::GetTriCoordColor(FVectorVMContext& Context
 	}
 }
 
-template<typename VertexAccessorType, typename TriType, typename BaryXType, typename BaryYType, typename BaryZType, typename UVSetType>
+template<typename VertexAccessorType>
 void UNiagaraDataInterfaceStaticMesh::GetTriCoordUV(FVectorVMContext& Context)
 {
-	TriType TriParam(Context);
-	BaryXType BaryXParam(Context);
-	BaryYType BaryYParam(Context);
-	BaryZType BaryZParam(Context);
-	UVSetType UVSetParam(Context);
-	FUserPtrHandler<FNDIStaticMesh_InstanceData> InstData(Context);
+	VectorVM::FExternalFuncInputHandler<int32> TriParam(Context);
+	VectorVM::FExternalFuncInputHandler<float> BaryXParam(Context);
+	VectorVM::FExternalFuncInputHandler<float> BaryYParam(Context);
+	VectorVM::FExternalFuncInputHandler<float> BaryZParam(Context);
+	VectorVM::FExternalFuncInputHandler<int32> UVSetParam(Context);
+	VectorVM::FUserPtrHandler<FNDIStaticMesh_InstanceData> InstData(Context);
 
-	FRegisterHandler<float> OutU(Context);
-	FRegisterHandler<float> OutV(Context);
+	VectorVM::FExternalFuncRegisterHandler<float> OutU(Context);
+	VectorVM::FExternalFuncRegisterHandler<float> OutV(Context);
 
 	FStaticMeshLODResources& Res = InstData->Mesh->RenderData->LODResources[0];
 	const FIndexArrayView& Indices = Res.IndexBuffer.GetArrayView();
@@ -1381,21 +1379,20 @@ void UNiagaraDataInterfaceStaticMesh::GetTriCoordUV(FVectorVMContext& Context)
 	}
 }
 
-template<typename TriType, typename BaryXType, typename BaryYType, typename BaryZType>
 void UNiagaraDataInterfaceStaticMesh::GetTriCoordPositionAndVelocity(FVectorVMContext& Context)
 {
-	TriType TriParam(Context);
-	BaryXType BaryXParam(Context);
-	BaryYType BaryYParam(Context);
-	BaryZType BaryZParam(Context);
-	FUserPtrHandler<FNDIStaticMesh_InstanceData> InstData(Context);
-
-	FRegisterHandler<float> OutPosX(Context);
-	FRegisterHandler<float> OutPosY(Context);
-	FRegisterHandler<float> OutPosZ(Context);
-	FRegisterHandler<float> OutVelX(Context);
-	FRegisterHandler<float> OutVelY(Context);
-	FRegisterHandler<float> OutVelZ(Context);
+	VectorVM::FExternalFuncInputHandler<int32> TriParam(Context);
+	VectorVM::FExternalFuncInputHandler<float> BaryXParam(Context);
+	VectorVM::FExternalFuncInputHandler<float> BaryYParam(Context);
+	VectorVM::FExternalFuncInputHandler<float> BaryZParam(Context);
+	VectorVM::FUserPtrHandler<FNDIStaticMesh_InstanceData> InstData(Context);
+
+	VectorVM::FExternalFuncRegisterHandler<float> OutPosX(Context);
+	VectorVM::FExternalFuncRegisterHandler<float> OutPosY(Context);
+	VectorVM::FExternalFuncRegisterHandler<float> OutPosZ(Context);
+	VectorVM::FExternalFuncRegisterHandler<float> OutVelX(Context);
+	VectorVM::FExternalFuncRegisterHandler<float> OutVelY(Context);
+	VectorVM::FExternalFuncRegisterHandler<float> OutVelZ(Context);
 
 	FStaticMeshLODResources& Res = InstData->Mesh->RenderData->LODResources[0];
 	const FIndexArrayView& Indices = Res.IndexBuffer.GetArrayView();
@@ -1436,22 +1433,22 @@ void UNiagaraDataInterfaceStaticMesh::GetTriCoordPositionAndVelocity(FVectorVMCo
 
 void UNiagaraDataInterfaceStaticMesh::WriteTransform(const FMatrix& ToWrite, FVectorVMContext& Context)
 {
-	FRegisterHandler<float> Out00(Context);
-	FRegisterHandler<float> Out01(Context);
-	FRegisterHandler<float> Out02(Context);
-	FRegisterHandler<float> Out03(Context);
-	FRegisterHandler<float> Out04(Context);
-	FRegisterHandler<float> Out05(Context);
-	FRegisterHandler<float> Out06(Context);
-	FRegisterHandler<float> Out07(Context);
-	FRegisterHandler<float> Out08(Context);
-	FRegisterHandler<float> Out09(Context);
-	FRegisterHandler<float> Out10(Context);
-	FRegisterHandler<float> Out11(Context);
-	FRegisterHandler<float> Out12(Context);
-	FRegisterHandler<float> Out13(Context);
-	FRegisterHandler<float> Out14(Context);
-	FRegisterHandler<float> Out15(Context);
+	VectorVM::FExternalFuncRegisterHandler<float> Out00(Context);
+	VectorVM::FExternalFuncRegisterHandler<float> Out01(Context);
+	VectorVM::FExternalFuncRegisterHandler<float> Out02(Context);
+	VectorVM::FExternalFuncRegisterHandler<float> Out03(Context);
+	VectorVM::FExternalFuncRegisterHandler<float> Out04(Context);
+	VectorVM::FExternalFuncRegisterHandler<float> Out05(Context);
+	VectorVM::FExternalFuncRegisterHandler<float> Out06(Context);
+	VectorVM::FExternalFuncRegisterHandler<float> Out07(Context);
+	VectorVM::FExternalFuncRegisterHandler<float> Out08(Context);
+	VectorVM::FExternalFuncRegisterHandler<float> Out09(Context);
+	VectorVM::FExternalFuncRegisterHandler<float> Out10(Context);
+	VectorVM::FExternalFuncRegisterHandler<float> Out11(Context);
+	VectorVM::FExternalFuncRegisterHandler<float> Out12(Context);
+	VectorVM::FExternalFuncRegisterHandler<float> Out13(Context);
+	VectorVM::FExternalFuncRegisterHandler<float> Out14(Context);
+	VectorVM::FExternalFuncRegisterHandler<float> Out15(Context);
 
 	for (int32 i = 0; i < Context.NumInstances; ++i)
 	{
@@ -1476,23 +1473,23 @@ void UNiagaraDataInterfaceStaticMesh::WriteTransform(const FMatrix& ToWrite, FVe
 
 void UNiagaraDataInterfaceStaticMesh::GetLocalToWorld(FVectorVMContext& Context)
 {
-	FUserPtrHandler<FNDIStaticMesh_InstanceData> InstData(Context);
+	VectorVM::FUserPtrHandler<FNDIStaticMesh_InstanceData> InstData(Context);
 	WriteTransform(InstData->Transform, Context);
 }
 
 void UNiagaraDataInterfaceStaticMesh::GetLocalToWorldInverseTransposed(FVectorVMContext& Context)
 {
-	FUserPtrHandler<FNDIStaticMesh_InstanceData> InstData(Context);
+	VectorVM::FUserPtrHandler<FNDIStaticMesh_InstanceData> InstData(Context);
 	WriteTransform(InstData->TransformInverseTransposed, Context);
 }
 
 void UNiagaraDataInterfaceStaticMesh::GetWorldVelocity(FVectorVMContext& Context)
 {
-	FUserPtrHandler<FNDIStaticMesh_InstanceData> InstData(Context);
+	VectorVM::FUserPtrHandler<FNDIStaticMesh_InstanceData> InstData(Context);
 
-	FRegisterHandler<float> OutVelX(Context);
-	FRegisterHandler<float> OutVelY(Context);
-	FRegisterHandler<float> OutVelZ(Context);
+	VectorVM::FExternalFuncRegisterHandler<float> OutVelX(Context);
+	VectorVM::FExternalFuncRegisterHandler<float> OutVelY(Context);
+	VectorVM::FExternalFuncRegisterHandler<float> OutVelZ(Context);
 
 	FVector Velocity(0.0f, 0.0f, 0.0f);
 	float InvDeltaTime = 1.0f / InstData->DeltaSeconds;
@@ -1513,16 +1510,16 @@ void UNiagaraDataInterfaceStaticMesh::GetWorldVelocity(FVectorVMContext& Context
 	}
 }
 
-template<typename TransformHandlerType, typename VertexIndexType>
+template<typename TransformHandlerType>
 void UNiagaraDataInterfaceStaticMesh::GetVertexPosition(FVectorVMContext& Context)
 {
 	TransformHandlerType TransformHandler;
-	VertexIndexType VertexIndexParam(Context);
-	FUserPtrHandler<FNDIStaticMesh_InstanceData> InstData(Context);
+	VectorVM::FExternalFuncInputHandler<int32> VertexIndexParam(Context);
+	VectorVM::FUserPtrHandler<FNDIStaticMesh_InstanceData> InstData(Context);
 
-	FRegisterHandler<float> OutPosX(Context);
-	FRegisterHandler<float> OutPosY(Context);
-	FRegisterHandler<float> OutPosZ(Context);
+	VectorVM::FExternalFuncRegisterHandler<float> OutPosX(Context);
+	VectorVM::FExternalFuncRegisterHandler<float> OutPosY(Context);
+	VectorVM::FExternalFuncRegisterHandler<float> OutPosZ(Context);
 
 	FStaticMeshLODResources& Res = InstData->Mesh->RenderData->LODResources[0];
 	const FPositionVertexBuffer& Positions = Res.VertexBuffers.PositionVertexBuffer;
diff --git a/Engine/Plugins/FX/Niagara/Source/Niagara/Private/NiagaraDataInterfaceTexture.cpp b/Engine/Plugins/FX/Niagara/Source/Niagara/Private/NiagaraDataInterfaceTexture.cpp
--- a/Engine/Plugins/FX/Niagara/Source/Niagara/Private/NiagaraDataInterfaceTexture.cpp
+++ b/Engine/Plugins/FX/Niagara/Source/Niagara/Private/NiagaraDataInterfaceTexture.cpp
@@ -165,13 +165,13 @@ void UNiagaraDataInterfaceTexture::GetFunctions(TArray<FNiagaraFunctionSignature
 		OutFunctions.Add(Sig);
 	}
 }
-DEFINE_NDI_RAW_FUNC_BINDER(UNiagaraDataInterfaceTexture, SampleTexture);
+DEFINE_NDI_DIRECT_FUNC_BINDER(UNiagaraDataInterfaceTexture, SampleTexture);
 void UNiagaraDataInterfaceTexture::GetVMExternalFunction(const FVMExternalFunctionBindingInfo& BindingInfo, void* InstanceData, FVMExternalFunction &OutFunc)
 {
 	if (BindingInfo.Name == SampleTexture2DName)
 	{
 		check(BindingInfo.GetNumInputs() == 2 && BindingInfo.GetNumOutputs() == 4);
-		TNDIParamBinder<0, float, TNDIParamBinder<1, float, NDI_RAW_FUNC_BINDER(UNiagaraDataInterfaceTexture, SampleTexture)>>::Bind(this, BindingInfo, InstanceData, OutFunc);
+		NDI_FUNC_BINDER(UNiagaraDataInterfaceTexture, SampleTexture)::Bind(this, OutFunc);
 	}
 	else if (BindingInfo.Name == SamplePseudoVolumeTextureName)
 	{
@@ -194,8 +194,8 @@ void UNiagaraDataInterfaceTexture::NoOp(FVectorVMContext& Context)
 
 void UNiagaraDataInterfaceTexture::GetTextureDimensions(FVectorVMContext& Context)
 {
-	FRegisterHandler<float> OutWidth(Context);
-	FRegisterHandler<float> OutHeight(Context);
+	VectorVM::FExternalFuncRegisterHandler<float> OutWidth(Context);
+	VectorVM::FExternalFuncRegisterHandler<float> OutHeight(Context);
 
 	if (Texture == nullptr)
 	{
@@ -217,15 +217,14 @@ void UNiagaraDataInterfaceTexture::GetTextureDimensions(FVectorVMContext& Contex
 	}
 }
 
-template<typename XType, typename YType>
 void UNiagaraDataInterfaceTexture::SampleTexture(FVectorVMContext& Context)
 {
-	XType XParam(Context);
-	YType YParam(Context);
-	FRegisterHandler<float> OutSampleR(Context);
-	FRegisterHandler<float> OutSampleG(Context);
-	FRegisterHandler<float> OutSampleB(Context);
-	FRegisterHandler<float> OutSampleA(Context);
+	VectorVM::FExternalFuncInputHandler<float> XParam(Context);
+	VectorVM::FExternalFuncInputHandler<float> YParam(Context);
+	VectorVM::FExternalFuncRegisterHandler<float> OutSampleR(Context);
+	VectorVM::FExternalFuncRegisterHandler<float> OutSampleG(Context);
+	VectorVM::FExternalFuncRegisterHandler<float> OutSampleB(Context);
+	VectorVM::FExternalFuncRegisterHandler<float> OutSampleA(Context);
 
 	for (int32 i = 0; i < Context.NumInstances; ++i)
 	{
diff --git a/Engine/Plugins/FX/Niagara/Source/Niagara/Private/NiagaraDataInterfaceVector2DCurve.cpp b/Engine/Plugins/FX/Niagara/Source/Niagara/Private/NiagaraDataInterfaceVector2DCurve.cpp
--- a/Engine/Plugins/FX/Niagara/Source/Niagara/Private/NiagaraDataInterfaceVector2DCurve.cpp
+++ b/Engine/Plugins/FX/Niagara/Source/Niagara/Private/NiagaraDataInterfaceVector2DCurve.cpp
@@ -157,12 +157,12 @@ void %s(in float In_X, out float2 Out_Value) \n\
 	return true;
 }
 
-DEFINE_NDI_RAW_FUNC_BINDER(UNiagaraDataInterfaceVector2DCurve, SampleCurve);
+DEFINE_NDI_FUNC_BINDER(UNiagaraDataInterfaceVector2DCurve, SampleCurve);
 void UNiagaraDataInterfaceVector2DCurve::GetVMExternalFunction(const FVMExternalFunctionBindingInfo& BindingInfo, void* InstanceData, FVMExternalFunction &OutFunc)
 {
 	if (BindingInfo.Name == SampleCurveName && BindingInfo.GetNumInputs() == 1 && BindingInfo.GetNumOutputs() == 2)
 	{
-		TCurveUseLUTBinder<TNDIParamBinder<0, float, NDI_RAW_FUNC_BINDER(UNiagaraDataInterfaceVector2DCurve, SampleCurve)>>::Bind(this, BindingInfo, InstanceData, OutFunc);
+		TCurveUseLUTBinder<NDI_FUNC_BINDER(UNiagaraDataInterfaceVector2DCurve, SampleCurve)>::Bind(this, BindingInfo, InstanceData, OutFunc);
 	}
 	else
 	{
@@ -192,13 +192,13 @@ FORCEINLINE_DEBUGGABLE FVector2D UNiagaraDataInterfaceVector2DCurve::SampleCurve
 	return FVector2D(XCurve.Eval(X), YCurve.Eval(X));
 }
 
-template<typename UseLUT, typename XParamType>
+template<typename UseLUT>
 void UNiagaraDataInterfaceVector2DCurve::SampleCurve(FVectorVMContext& Context)
 {
 	//TODO: Create some SIMDable optimized representation of the curve to do this faster.
-	XParamType XParam(Context);
-	FRegisterHandler<float> OutSampleX(Context);
-	FRegisterHandler<float> OutSampleY(Context);
+	VectorVM::FExternalFuncInputHandler<float> XParam(Context);
+	VectorVM::FExternalFuncRegisterHandler<float> OutSampleX(Context);
+	VectorVM::FExternalFuncRegisterHandler<float> OutSampleY(Context);
 
 	for (int32 i = 0; i < Context.NumInstances; ++i)
 	{
diff --git a/Engine/Plugins/FX/Niagara/Source/Niagara/Private/NiagaraDataInterfaceVector4Curve.cpp b/Engine/Plugins/FX/Niagara/Source/Niagara/Private/NiagaraDataInterfaceVector4Curve.cpp
--- a/Engine/Plugins/FX/Niagara/Source/Niagara/Private/NiagaraDataInterfaceVector4Curve.cpp
+++ b/Engine/Plugins/FX/Niagara/Source/Niagara/Private/NiagaraDataInterfaceVector4Curve.cpp
@@ -172,12 +172,12 @@ void %s(in float In_X, out float4 Out_Value) \n\
 	return true;
 }
 
-DEFINE_NDI_RAW_FUNC_BINDER(UNiagaraDataInterfaceVector4Curve, SampleCurve);
+DEFINE_NDI_FUNC_BINDER(UNiagaraDataInterfaceVector4Curve, SampleCurve);
 void UNiagaraDataInterfaceVector4Curve::GetVMExternalFunction(const FVMExternalFunctionBindingInfo& BindingInfo, void* InstanceData, FVMExternalFunction &OutFunc)
 {
 	if (BindingInfo.Name == SampleCurveName && BindingInfo.GetNumInputs() == 1 && BindingInfo.GetNumOutputs() == 4)
 	{
-		TCurveUseLUTBinder<TNDIParamBinder<0, float, NDI_RAW_FUNC_BINDER(UNiagaraDataInterfaceVector4Curve, SampleCurve)>>::Bind(this, BindingInfo, InstanceData, OutFunc);
+		TCurveUseLUTBinder<NDI_FUNC_BINDER(UNiagaraDataInterfaceVector4Curve, SampleCurve)>::Bind(this, BindingInfo, InstanceData, OutFunc);
 	}
 	else
 	{
@@ -208,15 +208,15 @@ FORCEINLINE_DEBUGGABLE FVector4 UNiagaraDataInterfaceVector4Curve::SampleCurveIn
 }
 
 
-template<typename UseLUT, typename XParamType>
+template<typename UseLUT>
 void UNiagaraDataInterfaceVector4Curve::SampleCurve(FVectorVMContext& Context)
 {
 	//TODO: Create some SIMDable optimized representation of the curve to do this faster.
-	XParamType XParam(Context);
-	FRegisterHandler<float> SamplePtrR(Context);
-	FRegisterHandler<float> SamplePtrG(Context);
-	FRegisterHandler<float> SamplePtrB(Context);
-	FRegisterHandler<float> SamplePtrA(Context);
+	VectorVM::FExternalFuncInputHandler<float> XParam(Context);
+	VectorVM::FExternalFuncRegisterHandler<float> SamplePtrR(Context);
+	VectorVM::FExternalFuncRegisterHandler<float> SamplePtrG(Context);
+	VectorVM::FExternalFuncRegisterHandler<float> SamplePtrB(Context);
+	VectorVM::FExternalFuncRegisterHandler<float> SamplePtrA(Context);
 
 	for (int32 i = 0; i < Context.NumInstances; ++i)
 	{
diff --git a/Engine/Plugins/FX/Niagara/Source/Niagara/Private/NiagaraDataInterfaceVectorCurve.cpp b/Engine/Plugins/FX/Niagara/Source/Niagara/Private/NiagaraDataInterfaceVectorCurve.cpp
--- a/Engine/Plugins/FX/Niagara/Source/Niagara/Private/NiagaraDataInterfaceVectorCurve.cpp
+++ b/Engine/Plugins/FX/Niagara/Source/Niagara/Private/NiagaraDataInterfaceVectorCurve.cpp
@@ -167,12 +167,12 @@ void %s(in float In_X, out float3 Out_Value) \n\
 	return true;
 }
 
-DEFINE_NDI_RAW_FUNC_BINDER(UNiagaraDataInterfaceVectorCurve, SampleCurve);
+DEFINE_NDI_FUNC_BINDER(UNiagaraDataInterfaceVectorCurve, SampleCurve);
 void UNiagaraDataInterfaceVectorCurve::GetVMExternalFunction(const FVMExternalFunctionBindingInfo& BindingInfo, void* InstanceData, FVMExternalFunction &OutFunc)
 {
 	if (BindingInfo.Name == SampleCurveName && BindingInfo.GetNumInputs() == 1 && BindingInfo.GetNumOutputs() == 3)
 	{
-		TCurveUseLUTBinder<TNDIParamBinder<0, float, NDI_RAW_FUNC_BINDER(UNiagaraDataInterfaceVectorCurve, SampleCurve)>>::Bind(this, BindingInfo, InstanceData, OutFunc);
+		TCurveUseLUTBinder<NDI_FUNC_BINDER(UNiagaraDataInterfaceVectorCurve, SampleCurve)>::Bind(this, BindingInfo, InstanceData, OutFunc);
 	}
 	else
 	{
@@ -202,14 +202,14 @@ FORCEINLINE_DEBUGGABLE FVector UNiagaraDataInterfaceVectorCurve::SampleCurveInte
 	return FVector(XCurve.Eval(X), YCurve.Eval(X), ZCurve.Eval(X));
 }
 
-template<typename UseLUT, typename XParamType>
+template<typename UseLUT>
 void UNiagaraDataInterfaceVectorCurve::SampleCurve(FVectorVMContext& Context)
 {
 	//TODO: Create some SIMDable optimized representation of the curve to do this faster.
-	XParamType XParam(Context);
-	FRegisterHandler<float> OutSampleX(Context);
-	FRegisterHandler<float> OutSampleY(Context);
-	FRegisterHandler<float> OutSampleZ(Context);
+	VectorVM::FExternalFuncInputHandler<float> XParam(Context);
+	VectorVM::FExternalFuncRegisterHandler<float> OutSampleX(Context);
+	VectorVM::FExternalFuncRegisterHandler<float> OutSampleY(Context);
+	VectorVM::FExternalFuncRegisterHandler<float> OutSampleZ(Context);
 
 	for (int32 i = 0; i < Context.NumInstances; ++i)
 	{
diff --git a/Engine/Plugins/FX/Niagara/Source/Niagara/Private/NiagaraDataInterfaceVectorField.cpp b/Engine/Plugins/FX/Niagara/Source/Niagara/Private/NiagaraDataInterfaceVectorField.cpp
--- a/Engine/Plugins/FX/Niagara/Source/Niagara/Private/NiagaraDataInterfaceVectorField.cpp
+++ b/Engine/Plugins/FX/Niagara/Source/Niagara/Private/NiagaraDataInterfaceVectorField.cpp
@@ -378,12 +378,12 @@ void UNiagaraDataInterfaceVectorField::GetFunctions(TArray<FNiagaraFunctionSigna
 	}
 }
 
-DEFINE_NDI_RAW_FUNC_BINDER(UNiagaraDataInterfaceVectorField, SampleVectorField);
+DEFINE_NDI_DIRECT_FUNC_BINDER(UNiagaraDataInterfaceVectorField, SampleVectorField);
 void UNiagaraDataInterfaceVectorField::GetVMExternalFunction(const FVMExternalFunctionBindingInfo& BindingInfo, void* InstanceData, FVMExternalFunction &OutFunc)
 {
 	if (BindingInfo.Name == SampleVectorFieldName && BindingInfo.GetNumInputs() == 3 && BindingInfo.GetNumOutputs() == 3)
 	{
-		TNDIParamBinder<0, float, TNDIParamBinder<1, float, TNDIParamBinder<2, float, NDI_RAW_FUNC_BINDER(UNiagaraDataInterfaceVectorField, SampleVectorField)>>>::Bind(this, BindingInfo, InstanceData, OutFunc);
+		NDI_FUNC_BINDER(UNiagaraDataInterfaceVectorField, SampleVectorField)::Bind(this, OutFunc);
 	}
 	else if (BindingInfo.Name == GetVectorDimsName && BindingInfo.GetNumInputs() == 0 && BindingInfo.GetNumOutputs() == 3)
 	{
@@ -397,15 +397,15 @@ void UNiagaraDataInterfaceVectorField::GetVMExternalFunction(const FVMExternalFu
 
 void UNiagaraDataInterfaceVectorField::GetFieldDimensions(FVectorVMContext& Context)
 {
-	/*FUserPtrHandler<FNDIVectorField_InstanceData> InstData(Context);
+	/*VectorVM::FUserPtrHandler<FNDIVectorField_InstanceData> InstData(Context);
 	if (!InstData.Get())
 	{
 		UE_LOG(LogNiagara, Warning, TEXT("FNDIVectorField_InstanceData has invalid instance data. %s"), *GetPathName());
 	}*/
 
-	FRegisterHandler<float> OutSizeX(Context);	
-	FRegisterHandler<float> OutSizeY(Context);
-	FRegisterHandler<float> OutSizeZ(Context);
+	VectorVM::FExternalFuncRegisterHandler<float> OutSizeX(Context);	
+	VectorVM::FExternalFuncRegisterHandler<float> OutSizeY(Context);
+	VectorVM::FExternalFuncRegisterHandler<float> OutSizeZ(Context);
 	
 	for (int32 i = 0; i < Context.NumInstances; ++i)
 	{
@@ -421,18 +421,18 @@ void UNiagaraDataInterfaceVectorField::GetFieldDimensions(FVectorVMContext& Cont
 
 void UNiagaraDataInterfaceVectorField::GetFieldBounds(FVectorVMContext& Context)
 {
-	/*FUserPtrHandler<FNDIVectorField_InstanceData> InstData(Context);
+	/*VectorVM::FUserPtrHandler<FNDIVectorField_InstanceData> InstData(Context);
 	if (!InstData.Get())
 	{
 		UE_LOG(LogNiagara, Warning, TEXT("FNDIVectorField_InstanceData has invalid instance data. %s"), *GetPathName());
 	}*/
 
-	FRegisterHandler<float> OutMinX(Context);
-	FRegisterHandler<float> OutMinY(Context);
-	FRegisterHandler<float> OutMinZ(Context);
-	FRegisterHandler<float> OutMaxX(Context);
-	FRegisterHandler<float> OutMaxY(Context);
-	FRegisterHandler<float> OutMaxZ(Context);
+	VectorVM::FExternalFuncRegisterHandler<float> OutMinX(Context);
+	VectorVM::FExternalFuncRegisterHandler<float> OutMinY(Context);
+	VectorVM::FExternalFuncRegisterHandler<float> OutMinZ(Context);
+	VectorVM::FExternalFuncRegisterHandler<float> OutMaxX(Context);
+	VectorVM::FExternalFuncRegisterHandler<float> OutMaxY(Context);
+	VectorVM::FExternalFuncRegisterHandler<float> OutMaxZ(Context);
 	
 	for (int32 i = 0; i < Context.NumInstances; ++i)
 	{
@@ -560,22 +560,20 @@ void UNiagaraDataInterfaceVectorField::Unlock()
 	}
 }
 
-
-template<typename XType, typename YType, typename ZType>
 void UNiagaraDataInterfaceVectorField::SampleVectorField(FVectorVMContext& Context)
 {
 	// Input arguments...
-	XType XParam(Context);
-	YType YParam(Context);
-	ZType ZParam(Context);
+	VectorVM::FExternalFuncInputHandler<float> XParam(Context);
+	VectorVM::FExternalFuncInputHandler<float> YParam(Context);
+	VectorVM::FExternalFuncInputHandler<float> ZParam(Context);
 
 	// User pointer buffer (as requested)
-	//FUserPtrHandler<FNDIVectorField_InstanceData> InstData(Context);
+	//VectorVM::FUserPtrHandler<FNDIVectorField_InstanceData> InstData(Context);
 
 	// Outputs...
-	FRegisterHandler<float> OutSampleX(Context);
-	FRegisterHandler<float> OutSampleY(Context);
-	FRegisterHandler<float> OutSampleZ(Context);
+	VectorVM::FExternalFuncRegisterHandler<float> OutSampleX(Context);
+	VectorVM::FExternalFuncRegisterHandler<float> OutSampleY(Context);
+	VectorVM::FExternalFuncRegisterHandler<float> OutSampleZ(Context);
 
 	/*UE_LOG(LogNiagara, Log, TEXT("Registers: In: %d %d %d %d Out: %d %d %d"),
 		XParam.RegisterIndex, YParam.RegisterIndex, ZParam.RegisterIndex, InstData.RegisterIndex,
diff --git a/Engine/Plugins/FX/Niagara/Source/NiagaraCore/Private/NiagaraCustomVersion.cpp b/Engine/Plugins/FX/Niagara/Source/NiagaraCore/Private/NiagaraCustomVersion.cpp
--- a/Engine/Plugins/FX/Niagara/Source/NiagaraCore/Private/NiagaraCustomVersion.cpp
+++ b/Engine/Plugins/FX/Niagara/Source/NiagaraCore/Private/NiagaraCustomVersion.cpp
@@ -8,4 +8,4 @@ const FGuid FNiagaraCustomVersion::GUID(0xFCF57AFA, 0x50764283, 0xB9A9E658, 0xFF
 // Register the custom version with core
 FCustomVersionRegistration GRegisterNiagaraCustomVersion(FNiagaraCustomVersion::GUID, FNiagaraCustomVersion::LatestVersion, TEXT("NiagaraVer"));
 
-const FGuid FNiagaraCustomVersion::LatestScriptCompileVersion(0x6DB2DD7A, 0xBE944FC7, 0xB2409AB0, 0xC13AF894);
\ No newline at end of file
+const FGuid FNiagaraCustomVersion::LatestScriptCompileVersion(0xF109EF07, 0xE00C47AC, 0xB1DF217F, 0xAC0F1789);
\ No newline at end of file
diff --git a/Engine/Plugins/FX/Niagara/Source/NiagaraEditor/Private/NiagaraHlslTranslator.cpp b/Engine/Plugins/FX/Niagara/Source/NiagaraEditor/Private/NiagaraHlslTranslator.cpp
--- a/Engine/Plugins/FX/Niagara/Source/NiagaraEditor/Private/NiagaraHlslTranslator.cpp
+++ b/Engine/Plugins/FX/Niagara/Source/NiagaraEditor/Private/NiagaraHlslTranslator.cpp
@@ -43,6 +43,8 @@
 #include "ShaderCore.h"
 #include "NiagaraShaderCompilationManager.h"
 
+#include "NiagaraEditorSettings.h"
+
 #define LOCTEXT_NAMESPACE "NiagaraCompiler"
 
 DECLARE_CYCLE_STAT(TEXT("Niagara - HlslTranslator - Translate"), STAT_NiagaraEditor_HlslTranslator_Translate, STATGROUP_NiagaraEditor);
@@ -2102,362 +2104,8 @@ void FHlslNiagaraTranslator::DecomposeVariableAccess(UStruct* Struct, bool bRead
 	}
 };
 
-TMap<FString, FString> FHlslNiagaraTranslator::ReplacementsForInvalid;
-
 void FHlslNiagaraTranslator::Init()
 {
-	if (ReplacementsForInvalid.Num() == 0)
-	{
-		// hlsl keywords
-		ReplacementsForInvalid.Add(TEXT("AppendStructuredBuffer"), TEXT("APPENDSTRUCTUREDBUFFER_VAR"));
-		ReplacementsForInvalid.Add(TEXT("asm"), TEXT("ASSEMBLY_VAR"));
-		ReplacementsForInvalid.Add(TEXT("asm_fragment"), TEXT("ASSEMBLY_FRAGMENT_VAR"));
-		ReplacementsForInvalid.Add(TEXT("BlendState"), TEXT("BLEND_STATE_VAR"));
-		ReplacementsForInvalid.Add(TEXT("bool"), TEXT("BOOL_VAR"));
-		ReplacementsForInvalid.Add(TEXT("bool1"), TEXT("BOOL1_VAR"));
-		ReplacementsForInvalid.Add(TEXT("bool2"), TEXT("BOOL2_VAR"));
-		ReplacementsForInvalid.Add(TEXT("bool3"), TEXT("BOOL3_VAR"));
-		ReplacementsForInvalid.Add(TEXT("bool4"), TEXT("BOOL4_VAR"));
-		ReplacementsForInvalid.Add(TEXT("break"), TEXT("BREAK_VAR"));
-		ReplacementsForInvalid.Add(TEXT("Buffer"), TEXT("BUFFER_VAR"));
-		ReplacementsForInvalid.Add(TEXT("ByteAddressBuffer"), TEXT("BYTE_ADDRESS_BUFFER_VAR"));
-		ReplacementsForInvalid.Add(TEXT("case"), TEXT("CASE_VAR"));
-		ReplacementsForInvalid.Add(TEXT("cbuffer"), TEXT("CBUFFER_VAR"));
-		ReplacementsForInvalid.Add(TEXT("centroid"), TEXT("CENTROID_VAR"));
-		ReplacementsForInvalid.Add(TEXT("class"), TEXT("CLASS_VAR"));
-		ReplacementsForInvalid.Add(TEXT("column_major"), TEXT("COLUMN_MAJOR_VAR"));
-		ReplacementsForInvalid.Add(TEXT("compile"), TEXT("COMPILE_VAR"));
-		ReplacementsForInvalid.Add(TEXT("compile_fragment"), TEXT("COMPILE_FRAGMENT_VAR"));
-		ReplacementsForInvalid.Add(TEXT("CompileShader"), TEXT("COMPILE_SHADER_VAR"));
-		ReplacementsForInvalid.Add(TEXT("const"), TEXT("CONST_VAR"));
-		ReplacementsForInvalid.Add(TEXT("continue"), TEXT("CONTINUE_VAR"));
-		ReplacementsForInvalid.Add(TEXT("ComputeShader"), TEXT("COMPUTE_SHADER_VAR"));
-		ReplacementsForInvalid.Add(TEXT("ConsumeStructuredBuffer"), TEXT("CONSUME_STRUCTURED_BUFFER"));
-		ReplacementsForInvalid.Add(TEXT("default"), TEXT("DEFAULT_VAR"));
-		ReplacementsForInvalid.Add(TEXT("DepthStencilState"), TEXT("DEPTHSTENCILSTATE_VAR"));
-		ReplacementsForInvalid.Add(TEXT("DepthStencilView"), TEXT("DEPTHSTENCILVIEW_VAR"));
-		ReplacementsForInvalid.Add(TEXT("discard"), TEXT("DISCARD_VAR"));
-		ReplacementsForInvalid.Add(TEXT("do"), TEXT("DO_VAR"));
-		ReplacementsForInvalid.Add(TEXT("double"), TEXT("DOUBLE_VAR"));
-		ReplacementsForInvalid.Add(TEXT("double1"), TEXT("DOUBLE1_VAR"));
-		ReplacementsForInvalid.Add(TEXT("double2"), TEXT("DOUBLE2_VAR"));
-		ReplacementsForInvalid.Add(TEXT("double3"), TEXT("DOUBLE3_VAR"));
-		ReplacementsForInvalid.Add(TEXT("double4"), TEXT("DOUBLE4_VAR"));
-		ReplacementsForInvalid.Add(TEXT("DomainShader"), TEXT("DOMAINSHADER_VAR"));
-		ReplacementsForInvalid.Add(TEXT("dword"), TEXT("DWORD_VAR"));
-		ReplacementsForInvalid.Add(TEXT("else"), TEXT("ELSE_VAR"));
-		ReplacementsForInvalid.Add(TEXT("export"), TEXT("EXPORT_VAR"));
-		ReplacementsForInvalid.Add(TEXT("extern"), TEXT("EXTERN_VAR"));
-		ReplacementsForInvalid.Add(TEXT("false"), TEXT("FALSE_VAR"));
-		ReplacementsForInvalid.Add(TEXT("float"), TEXT("FLOAT_VAR"));
-		ReplacementsForInvalid.Add(TEXT("float1"), TEXT("FLOAT1_VAR"));
-		ReplacementsForInvalid.Add(TEXT("float2"), TEXT("FLOAT2_VAR"));
-		ReplacementsForInvalid.Add(TEXT("float3"), TEXT("FLOAT3_VAR"));
-		ReplacementsForInvalid.Add(TEXT("float4"), TEXT("FLOAT4_VAR"));
-		ReplacementsForInvalid.Add(TEXT("for"), TEXT("FOR_VAR"));
-		ReplacementsForInvalid.Add(TEXT("fxgroup"), TEXT("FXGROUP_VAR"));
-		ReplacementsForInvalid.Add(TEXT("GeometryShader"), TEXT("GEOMETRYSHADER_VAR"));
-		ReplacementsForInvalid.Add(TEXT("groupshared"), TEXT("GROUPSHARED_VAR"));
-		ReplacementsForInvalid.Add(TEXT("half"), TEXT("HALF_VAR"));
-		ReplacementsForInvalid.Add(TEXT("half1"), TEXT("HALF1_VAR"));
-		ReplacementsForInvalid.Add(TEXT("half2"), TEXT("HALF2_VAR"));
-		ReplacementsForInvalid.Add(TEXT("half3"), TEXT("HALF3_VAR"));
-		ReplacementsForInvalid.Add(TEXT("half4"), TEXT("HALF4_VAR"));
-		ReplacementsForInvalid.Add(TEXT("Hullshader"), TEXT("HULLSHADER_VAR"));
-		ReplacementsForInvalid.Add(TEXT("if"), TEXT("IF_VAR"));
-		ReplacementsForInvalid.Add(TEXT("in"), TEXT("IN_VAR"));
-		ReplacementsForInvalid.Add(TEXT("inline"), TEXT("INLINE_VAR"));
-		ReplacementsForInvalid.Add(TEXT("inout"), TEXT("INOUT_VAR"));
-		ReplacementsForInvalid.Add(TEXT("InputPatch"), TEXT("INPUTPATCH_VAR"));
-		ReplacementsForInvalid.Add(TEXT("int"), TEXT("INT_VAR"));
-		ReplacementsForInvalid.Add(TEXT("int1"), TEXT("INT1_VAR"));
-		ReplacementsForInvalid.Add(TEXT("int2"), TEXT("INT2_VAR"));
-		ReplacementsForInvalid.Add(TEXT("int3"), TEXT("INT3_VAR"));
-		ReplacementsForInvalid.Add(TEXT("int4"), TEXT("INT4_VAR"));
-		ReplacementsForInvalid.Add(TEXT("interface"), TEXT("INTERFACE_VAR"));
-		ReplacementsForInvalid.Add(TEXT("line"), TEXT("LINE_VAR"));
-		ReplacementsForInvalid.Add(TEXT("lineadj"), TEXT("LINEADJ_VAR"));
-		ReplacementsForInvalid.Add(TEXT("linear"), TEXT("LINEAR_VAR"));
-		ReplacementsForInvalid.Add(TEXT("LineStream"), TEXT("LINESTREAM_VAR"));
-		ReplacementsForInvalid.Add(TEXT("matrix"), TEXT("MATRIX_VAR"));
-		ReplacementsForInvalid.Add(TEXT("min16float"), TEXT("MIN16FLOAT_VAR"));
-		ReplacementsForInvalid.Add(TEXT("min10float"), TEXT("MIN10FLOAT_VAR"));
-		ReplacementsForInvalid.Add(TEXT("min16int"), TEXT("MIN16INT_VAR"));
-		ReplacementsForInvalid.Add(TEXT("min12int"), TEXT("MIN12INT_VAR"));
-		ReplacementsForInvalid.Add(TEXT("min16uint"), TEXT("MIN16UINT_VAR"));
-		ReplacementsForInvalid.Add(TEXT("namespace"), TEXT("NAMESPACE_VAR"));
-		ReplacementsForInvalid.Add(TEXT("nointerpolation"), TEXT("NOINTERPOLATION_VAR"));
-		ReplacementsForInvalid.Add(TEXT("noperspective"), TEXT("NOPERSPECTIVE_VAR"));
-		ReplacementsForInvalid.Add(TEXT("NULL"), TEXT("NULL_VAR"));
-		ReplacementsForInvalid.Add(TEXT("out"), TEXT("OUT_VAR"));
-		ReplacementsForInvalid.Add(TEXT("OutputPatch"), TEXT("OUTPUTPATCH_VAR"));
-		ReplacementsForInvalid.Add(TEXT("packoffset"), TEXT("PACKOFFSET,_VAR"));
-		ReplacementsForInvalid.Add(TEXT("pass"), TEXT("PASS,_VAR"));
-		ReplacementsForInvalid.Add(TEXT("pixelfragment"), TEXT("PIXELFRAGMENT_VAR"));
-		ReplacementsForInvalid.Add(TEXT("PixelShader"), TEXT("PIXELSHADER_VAR"));
-		ReplacementsForInvalid.Add(TEXT("point"), TEXT("POINT_VAR"));
-		ReplacementsForInvalid.Add(TEXT("PointStream"), TEXT("POINTSTREAM_VAR"));
-		ReplacementsForInvalid.Add(TEXT("precise"), TEXT("PRECISE_VAR"));
-		ReplacementsForInvalid.Add(TEXT("RasterizerState"), TEXT("RASTERIZERSTATE_VAR"));
-		ReplacementsForInvalid.Add(TEXT("RenderTargetView"), TEXT("RENDERTARGETVIEW_VAR"));
-		ReplacementsForInvalid.Add(TEXT("return"), TEXT("RETURN_VAR"));
-		ReplacementsForInvalid.Add(TEXT("register"), TEXT("REGISTER_VAR"));
-		ReplacementsForInvalid.Add(TEXT("row_major"), TEXT("ROW_MAJOR_VAR"));
-		ReplacementsForInvalid.Add(TEXT("RWBuffer"), TEXT("RWBUFFER_VAR"));
-		ReplacementsForInvalid.Add(TEXT("RWByteAddressBuffer"), TEXT("RWBYTEADDRESSBUFFER_VAR"));
-		ReplacementsForInvalid.Add(TEXT("RWStructuredBuffer"), TEXT("RWSTRUCTUREDBUFFER_VAR"));
-		ReplacementsForInvalid.Add(TEXT("RWTexture1D"), TEXT("RWTEXTURE1D_VAR"));
-		ReplacementsForInvalid.Add(TEXT("RWTexture1DArray"), TEXT("RWTEXTURE1DARRAY_VAR"));
-		ReplacementsForInvalid.Add(TEXT("RWTexture2D"), TEXT("RWTEXTURE2D_VAR"));
-		ReplacementsForInvalid.Add(TEXT("RWTexture2DArray"), TEXT("RWTEXTURE2DARRAY_VAR"));
-		ReplacementsForInvalid.Add(TEXT("RWTexture3D"), TEXT("RWTEXTURE3D_VAR"));
-		ReplacementsForInvalid.Add(TEXT("sample"), TEXT("SAMPLE_VAR"));
-		ReplacementsForInvalid.Add(TEXT("sampler"), TEXT("SAMPLER_VAR"));
-		ReplacementsForInvalid.Add(TEXT("SamplerState"), TEXT("SAMPLERSTATE_VAR"));
-		ReplacementsForInvalid.Add(TEXT("SamplerComparisonState"), TEXT("SAMPLERCOMPARISONSTATE_VAR"));
-		ReplacementsForInvalid.Add(TEXT("shared"), TEXT("SHARED_VAR"));
-		ReplacementsForInvalid.Add(TEXT("snorm"), TEXT("SNORM_VAR"));
-		ReplacementsForInvalid.Add(TEXT("stateblock"), TEXT("STATEBLOCK_VAR"));
-		ReplacementsForInvalid.Add(TEXT("stateblock_state"), TEXT("STATEBLOCK_STATE_VAR"));
-		ReplacementsForInvalid.Add(TEXT("static"), TEXT("STATIC_VAR"));
-		ReplacementsForInvalid.Add(TEXT("string"), TEXT("STRING_VAR"));
-		ReplacementsForInvalid.Add(TEXT("struct"), TEXT("STRUCT_VAR"));
-		ReplacementsForInvalid.Add(TEXT("switch"), TEXT("SWITCH_VAR"));
-		ReplacementsForInvalid.Add(TEXT("StructuredBuffer"), TEXT("STRUCTUREDBUFFER_VAR"));
-		ReplacementsForInvalid.Add(TEXT("tbuffer"), TEXT("TBUFFER_VAR"));
-		ReplacementsForInvalid.Add(TEXT("technique"), TEXT("TECHNIQUE_VAR"));
-		ReplacementsForInvalid.Add(TEXT("technique10"), TEXT("TECHNIQUE10_VAR"));
-		ReplacementsForInvalid.Add(TEXT("technique11"), TEXT("TECHNIQUE11_VAR"));
-		ReplacementsForInvalid.Add(TEXT("texture"), TEXT("TEXTURE_VAR"));
-		ReplacementsForInvalid.Add(TEXT("Texture1D"), TEXT("TEXTURE1D_VAR"));
-		ReplacementsForInvalid.Add(TEXT("Texture1DArray"), TEXT("TEXTURE1DARRAY_VAR"));
-		ReplacementsForInvalid.Add(TEXT("Texture2D"), TEXT("TEXTURE2D_VAR"));
-		ReplacementsForInvalid.Add(TEXT("Texture2DArray"), TEXT("TEXTURE2DARRAY_VAR"));
-		ReplacementsForInvalid.Add(TEXT("Texture2DMS"), TEXT("TEXTURE2DMS_VAR"));
-		ReplacementsForInvalid.Add(TEXT("Texture2DMSArray"), TEXT("TEXTURE2DMSARRAY_VAR"));
-		ReplacementsForInvalid.Add(TEXT("Texture3D"), TEXT("TEXTURE3D_VAR"));
-		ReplacementsForInvalid.Add(TEXT("TextureCube"), TEXT("TEXTURECUBE_VAR"));
-		ReplacementsForInvalid.Add(TEXT("TextureCubeArray"), TEXT("TEXTURECUBEARRAY_VAR"));
-		ReplacementsForInvalid.Add(TEXT("true"), TEXT("TRUE_VAR"));
-		ReplacementsForInvalid.Add(TEXT("typedef"), TEXT("TYPEDEF_VAR"));
-		ReplacementsForInvalid.Add(TEXT("triangle"), TEXT("TRIANGLE_VAR"));
-		ReplacementsForInvalid.Add(TEXT("triangleadj"), TEXT("TRIANGLEADJ_VAR"));
-		ReplacementsForInvalid.Add(TEXT("TriangleStream"), TEXT("TRIANGLESTREAM_VAR"));
-		ReplacementsForInvalid.Add(TEXT("uint"), TEXT("UINT_VAR"));
-		ReplacementsForInvalid.Add(TEXT("uint1"), TEXT("UINT1_VAR"));
-		ReplacementsForInvalid.Add(TEXT("uint2"), TEXT("UINT2_VAR"));
-		ReplacementsForInvalid.Add(TEXT("uint3"), TEXT("UINT3_VAR"));
-		ReplacementsForInvalid.Add(TEXT("uint4"), TEXT("UINT4_VAR"));
-		ReplacementsForInvalid.Add(TEXT("uniform"), TEXT("UNIFORM_VAR"));
-		ReplacementsForInvalid.Add(TEXT("unorm"), TEXT("UNORM_VAR"));
-		ReplacementsForInvalid.Add(TEXT("unsigned"), TEXT("UNSIGNED_VAR"));
-		ReplacementsForInvalid.Add(TEXT("vector"), TEXT("VECTOR_VAR"));
-		ReplacementsForInvalid.Add(TEXT("vertexfragment"), TEXT("VERTEXFRAGMENT_VAR"));
-		ReplacementsForInvalid.Add(TEXT("VertexShader"), TEXT("VERTEXSHADER_VAR"));
-		ReplacementsForInvalid.Add(TEXT("void"), TEXT("VOID_VAR"));
-		ReplacementsForInvalid.Add(TEXT("volatile"), TEXT("VOLATILE_VAR"));
-		ReplacementsForInvalid.Add(TEXT("while"), TEXT("WHILE_VAR"));
-
-		// additional glsl keywords
-		ReplacementsForInvalid.Add(TEXT("attribute"), TEXT("ATTRIBUTE_VAR"));
-		ReplacementsForInvalid.Add(TEXT("varying"), TEXT("VARYING_VAR"));
-		ReplacementsForInvalid.Add(TEXT("layout"), TEXT("LAYOUT_VAR"));
-		ReplacementsForInvalid.Add(TEXT("mat2"), TEXT("MAT2_VAR"));
-		ReplacementsForInvalid.Add(TEXT("mat3"), TEXT("MAT3_VAR"));
-		ReplacementsForInvalid.Add(TEXT("mat4"), TEXT("MAT4_VAR"));
-		ReplacementsForInvalid.Add(TEXT("mat2x2"), TEXT("MAT2X2_VAR"));
-		ReplacementsForInvalid.Add(TEXT("mat2x3"), TEXT("MAT2X3_VAR"));
-		ReplacementsForInvalid.Add(TEXT("mat2x4"), TEXT("MAT2X4_VAR"));
-		ReplacementsForInvalid.Add(TEXT("mat3x2"), TEXT("MAT3X2_VAR"));
-		ReplacementsForInvalid.Add(TEXT("mat3x3"), TEXT("MAT3X3_VAR"));
-		ReplacementsForInvalid.Add(TEXT("mat3x4"), TEXT("MAT3X4_VAR"));
-		ReplacementsForInvalid.Add(TEXT("mat4x2"), TEXT("MAT4X2_VAR"));
-		ReplacementsForInvalid.Add(TEXT("mat4x3"), TEXT("MAT4X3_VAR"));
-		ReplacementsForInvalid.Add(TEXT("mat4x4"), TEXT("MAT4X4_VAR"));
-		ReplacementsForInvalid.Add(TEXT("vec2"), TEXT("VEC2_VAR"));
-		ReplacementsForInvalid.Add(TEXT("vec3"), TEXT("VEC3_VAR"));
-		ReplacementsForInvalid.Add(TEXT("vec4"), TEXT("VEC4_VAR"));
-		ReplacementsForInvalid.Add(TEXT("ivec2"), TEXT("IVEC2_VAR"));
-		ReplacementsForInvalid.Add(TEXT("ivec3"), TEXT("IVEC3_VAR"));
-		ReplacementsForInvalid.Add(TEXT("ivec4"), TEXT("IVEC4_VAR"));
-		ReplacementsForInvalid.Add(TEXT("bvec2"), TEXT("BVEC2_VAR"));
-		ReplacementsForInvalid.Add(TEXT("bvec3"), TEXT("BVEC3_VAR"));
-		ReplacementsForInvalid.Add(TEXT("bvec4"), TEXT("BVEC4_VAR"));
-		ReplacementsForInvalid.Add(TEXT("uvec2"), TEXT("UVEC2_VAR"));
-		ReplacementsForInvalid.Add(TEXT("uvec3"), TEXT("UVEC3_VAR"));
-		ReplacementsForInvalid.Add(TEXT("uvec4"), TEXT("UVEC4_VAR"));
-		ReplacementsForInvalid.Add(TEXT("lowp"), TEXT("LOWP_VAR"));
-		ReplacementsForInvalid.Add(TEXT("mediump"), TEXT("MEDIUMP_VAR"));
-		ReplacementsForInvalid.Add(TEXT("highp"), TEXT("HIGHP_VAR"));
-		ReplacementsForInvalid.Add(TEXT("precision"), TEXT("PRECISION_VAR"));
-		ReplacementsForInvalid.Add(TEXT("sampler1D"), TEXT("SAMPLER1D_VAR"));
-		ReplacementsForInvalid.Add(TEXT("sampler2D"), TEXT("SAMPLER2D_VAR"));
-		ReplacementsForInvalid.Add(TEXT("sampler3D"), TEXT("SAMPLER3D_VAR"));
-		ReplacementsForInvalid.Add(TEXT("samplerCube"), TEXT("SAMPLERCUBE_VAR"));
-		ReplacementsForInvalid.Add(TEXT("sampler1DShadow"), TEXT("SAMPLER1DSHADOW_VAR"));
-		ReplacementsForInvalid.Add(TEXT("sampler2DShadow"), TEXT("SAMPLER2DSHADOW_VAR"));
-		ReplacementsForInvalid.Add(TEXT("samplerCubeShadow"), TEXT("SAMPLERCUBESHADOW_VAR"));
-		ReplacementsForInvalid.Add(TEXT("sampler1DArray"), TEXT("SAMPLER1DARRAY_VAR"));
-		ReplacementsForInvalid.Add(TEXT("sampler2DArray"), TEXT("SAMPLER2DARRAY_VAR"));
-		ReplacementsForInvalid.Add(TEXT("sampler1DArrayShadow"), TEXT("SAMPLER1DARRAYSHADOW_VAR"));
-		ReplacementsForInvalid.Add(TEXT("sampler2DArrayShadow"), TEXT("SAMPLER2DARRAYSHADOW_VAR"));
-		ReplacementsForInvalid.Add(TEXT("isampler1D"), TEXT("ISAMPLER1D_VAR"));
-		ReplacementsForInvalid.Add(TEXT("isampler2D"), TEXT("ISAMPLER2D_VAR"));
-		ReplacementsForInvalid.Add(TEXT("isampler3D"), TEXT("ISAMPLER3D_VAR"));
-		ReplacementsForInvalid.Add(TEXT("isamplerCube"), TEXT("ISAMPLERCUBE_VAR"));
-		ReplacementsForInvalid.Add(TEXT("isampler1DArray"), TEXT("ISAMPLER1DARRAY_VAR"));
-		ReplacementsForInvalid.Add(TEXT("isampler2DArray"), TEXT("ISAMPLER2DARRAY_VAR"));
-		ReplacementsForInvalid.Add(TEXT("usampler1D"), TEXT("USAMPLER1D_VAR"));
-		ReplacementsForInvalid.Add(TEXT("usampler2D"), TEXT("USAMPLER2D_VAR"));
-		ReplacementsForInvalid.Add(TEXT("usampler3D"), TEXT("USAMPLER3D_VAR"));
-		ReplacementsForInvalid.Add(TEXT("usamplerCube"), TEXT("USAMPLERCUBE_VAR"));
-		ReplacementsForInvalid.Add(TEXT("usampler1DArray"), TEXT("USAMPLER1DARRAY_VAR"));
-		ReplacementsForInvalid.Add(TEXT("usampler2DArray"), TEXT("USAMPLER2DARRAY_VAR"));
-		ReplacementsForInvalid.Add(TEXT("sampler2DRect"), TEXT("SAMPLER2DRECT_VAR"));
-		ReplacementsForInvalid.Add(TEXT("sampler2DRectShadow"), TEXT("SAMPLER2DRECTSHADOW_VAR"));
-		ReplacementsForInvalid.Add(TEXT("isampler2DRect"), TEXT("ISAMPLER2DRECT_VAR"));
-		ReplacementsForInvalid.Add(TEXT("usampler2DRect"), TEXT("USAMPLER2DRECT_VAR"));
-		ReplacementsForInvalid.Add(TEXT("samplerBuffer"), TEXT("SAMPLERBUFFER_VAR"));
-		ReplacementsForInvalid.Add(TEXT("isamplerBuffer"), TEXT("ISAMPLERBUFFER_VAR"));
-		ReplacementsForInvalid.Add(TEXT("usamplerBuffer"), TEXT("USAMPLERBUFFER_VAR"));
-		ReplacementsForInvalid.Add(TEXT("struct"), TEXT("STRUCT_VAR"));
-		ReplacementsForInvalid.Add(TEXT("flat"), TEXT("FLAT_VAR"));
-		ReplacementsForInvalid.Add(TEXT("smooth"), TEXT("SMOOTH_VAR"));
-		ReplacementsForInvalid.Add(TEXT("invariant"), TEXT("INVARIANT_VAR"));
-		ReplacementsForInvalid.Add(TEXT("union"), TEXT("UNION_VAR"));
-		ReplacementsForInvalid.Add(TEXT("enum"), TEXT("ENUM_VAR"));
-		ReplacementsForInvalid.Add(TEXT("typedef"), TEXT("TYPEDEF_VAR"));
-		ReplacementsForInvalid.Add(TEXT("template"), TEXT("TEMPLATE_VAR"));
-		ReplacementsForInvalid.Add(TEXT("this"), TEXT("THIS_VAR"));
-		ReplacementsForInvalid.Add(TEXT("packed"), TEXT("PACKED_VAR"));
-		ReplacementsForInvalid.Add(TEXT("goto"), TEXT("GOTO_VAR"));
-		ReplacementsForInvalid.Add(TEXT("noinline"), TEXT("NOINLINE_VAR"));
-		ReplacementsForInvalid.Add(TEXT("public"), TEXT("PUBLIC_VAR"));
-		ReplacementsForInvalid.Add(TEXT("external"), TEXT("EXTERNAL_VAR"));
-		ReplacementsForInvalid.Add(TEXT("long short"), TEXT("LONG SHORT_VAR"));
-		ReplacementsForInvalid.Add(TEXT("superp"), TEXT("SUPERP_VAR"));
-		ReplacementsForInvalid.Add(TEXT("input"), TEXT("INPUT_VAR"));
-		ReplacementsForInvalid.Add(TEXT("output"), TEXT("OUTPUT_VAR"));
-		ReplacementsForInvalid.Add(TEXT("hvec2"), TEXT("HVEC2_VAR"));
-		ReplacementsForInvalid.Add(TEXT("hvec3"), TEXT("HVEC3_VAR"));
-		ReplacementsForInvalid.Add(TEXT("hvec4"), TEXT("HVEC4_VAR"));
-		ReplacementsForInvalid.Add(TEXT("dvec2"), TEXT("DVEC2_VAR"));
-		ReplacementsForInvalid.Add(TEXT("dvec3"), TEXT("DVEC3_VAR"));
-		ReplacementsForInvalid.Add(TEXT("dvec4"), TEXT("DVEC4_VAR"));
-		ReplacementsForInvalid.Add(TEXT("fvec2"), TEXT("FVEC2_VAR"));
-		ReplacementsForInvalid.Add(TEXT("fvec3"), TEXT("FVEC3_VAR"));
-		ReplacementsForInvalid.Add(TEXT("fvec4"), TEXT("FVEC4_VAR"));
-		ReplacementsForInvalid.Add(TEXT("sampler3DRect"), TEXT("SAMPLER3DRECT_VAR"));
-		ReplacementsForInvalid.Add(TEXT("filter"), TEXT("FILTER_VAR"));
-		ReplacementsForInvalid.Add(TEXT("image1D"), TEXT("IMAGE1D_VAR"));
-		ReplacementsForInvalid.Add(TEXT("image2D"), TEXT("IMAGE2D_VAR"));
-		ReplacementsForInvalid.Add(TEXT("image3D"), TEXT("IMAGE3D_VAR"));
-		ReplacementsForInvalid.Add(TEXT("imageCube"), TEXT("IMAGECUBE_VAR"));
-		ReplacementsForInvalid.Add(TEXT("iimage1D"), TEXT("IIMAGE1D_VAR"));
-		ReplacementsForInvalid.Add(TEXT("iimage2D"), TEXT("IIMAGE2D_VAR"));
-		ReplacementsForInvalid.Add(TEXT("iimage3D"), TEXT("IIMAGE3D_VAR"));
-		ReplacementsForInvalid.Add(TEXT("iimageCub"), TEXT("IIMAGECUBE_VAR"));
-		ReplacementsForInvalid.Add(TEXT("uimage1D"), TEXT("UIMAGE1D_VAR"));
-		ReplacementsForInvalid.Add(TEXT("uimage2D"), TEXT("UIMAGE2D_VAR"));
-		ReplacementsForInvalid.Add(TEXT("uimage3D"), TEXT("UIMAGE3D_VAR"));
-		ReplacementsForInvalid.Add(TEXT("uimageCube"), TEXT("UIMAGECUBE_VAR"));
-		ReplacementsForInvalid.Add(TEXT("image1DArray"), TEXT("IMAGE1DARRAY_VAR"));
-		ReplacementsForInvalid.Add(TEXT("image2DArray"), TEXT("IMAGE2DARRAY_VAR"));
-		ReplacementsForInvalid.Add(TEXT("iimage1DArray"), TEXT("IIMAGE1DARRAY_VAR"));
-		ReplacementsForInvalid.Add(TEXT("iimage2DArray"), TEXT("IIMAGE2DARRAY_VAR"));
-		ReplacementsForInvalid.Add(TEXT("uimage1DArray"), TEXT("UIMAGE1DARRAY_VAR"));
-		ReplacementsForInvalid.Add(TEXT("uimage2DArray"), TEXT("UIMAGE2DARRAY_VAR"));
-		ReplacementsForInvalid.Add(TEXT("image1DShadow"), TEXT("IMAGE1DSHADOW_VAR"));
-		ReplacementsForInvalid.Add(TEXT("image2DShadow"), TEXT("IMAGE2DSHADOW_VAR"));
-		ReplacementsForInvalid.Add(TEXT("image1DArrayShadow"), TEXT("IMAGE1DARRAYSHADOW_VAR"));
-		ReplacementsForInvalid.Add(TEXT("image2DArrayShadow"), TEXT("IMAGE2DARRAYSHADOW_VAR"));
-		ReplacementsForInvalid.Add(TEXT("imageBuffer"), TEXT("IMAGEBUFFER_VAR"));
-		ReplacementsForInvalid.Add(TEXT("iimageBuffer"), TEXT("IIMAGEBUFFER_VAR"));
-		ReplacementsForInvalid.Add(TEXT("uimageBuffer"), TEXT("UIMAGEBUFFER_VAR"));
-		ReplacementsForInvalid.Add(TEXT("sizeof"), TEXT("SIZEOF_VAR"));
-		ReplacementsForInvalid.Add(TEXT("cast"), TEXT("CAST_VAR"));
-		ReplacementsForInvalid.Add(TEXT("using"), TEXT("USING_VAR"));
-		ReplacementsForInvalid.Add(TEXT("row_major"), TEXT("ROW_MAJOR_VAR"));
-		ReplacementsForInvalid.Add(TEXT("int8_t"), TEXT("INT8_T_VAR"));
-		ReplacementsForInvalid.Add(TEXT("unsigned char"), TEXT("UNSIGNED_CHAR_VAR"));
-		ReplacementsForInvalid.Add(TEXT("uchar"), TEXT("UCHAR_VAR"));
-		ReplacementsForInvalid.Add(TEXT("uint8_t"), TEXT("UINT8_T_VAR"));
-		ReplacementsForInvalid.Add(TEXT("short"), TEXT("SHORT_VAR"));
-		ReplacementsForInvalid.Add(TEXT("int16_t"), TEXT("INT16_T_VAR"));
-		ReplacementsForInvalid.Add(TEXT("unsigned shorT"), TEXT("UNSIGNED_SHORT_VAR"));
-		ReplacementsForInvalid.Add(TEXT("ushort"), TEXT("USHORT_VAR"));
-		ReplacementsForInvalid.Add(TEXT("unit16_t"), TEXT("UNIT16_T_VAR"));
-		ReplacementsForInvalid.Add(TEXT("int32_t"), TEXT("INT32_T_VAR"));
-		ReplacementsForInvalid.Add(TEXT("unsigned int"), TEXT("UNSIGNED_INT_VAR"));
-		ReplacementsForInvalid.Add(TEXT("uint32_t"), TEXT("UINT32_T_VAR"));
-		ReplacementsForInvalid.Add(TEXT("smooth"), TEXT("SMOOTH_VAR"));
-		ReplacementsForInvalid.Add(TEXT("invariant"), TEXT("INVARIANT_VAR"));
-		ReplacementsForInvalid.Add(TEXT("union"), TEXT("UNION_VAR"));
-		ReplacementsForInvalid.Add(TEXT("enum"), TEXT("ENUM_VAR"));
-		ReplacementsForInvalid.Add(TEXT("typedef"), TEXT("TYPEDEF_VAR"));
-		ReplacementsForInvalid.Add(TEXT("template"), TEXT("TEMPLATE_VAR"));
-		ReplacementsForInvalid.Add(TEXT("this"), TEXT("THIS_VAR"));
-		ReplacementsForInvalid.Add(TEXT("packed"), TEXT("PACKED_VAR"));
-		ReplacementsForInvalid.Add(TEXT("goto"), TEXT("GOTO_VAR"));
-		ReplacementsForInvalid.Add(TEXT("noinline"), TEXT("NOINLINE_VAR"));
-		ReplacementsForInvalid.Add(TEXT("public"), TEXT("PUBLIC_VAR"));
-		ReplacementsForInvalid.Add(TEXT("external"), TEXT("EXTERNAL_VAR"));
-		ReplacementsForInvalid.Add(TEXT(" long short"), TEXT("LONG_SHORT_VAR"));
-		ReplacementsForInvalid.Add(TEXT("superp"), TEXT("SUPERP_VAR"));
-		ReplacementsForInvalid.Add(TEXT("input"), TEXT("INPUT_VAR"));
-		ReplacementsForInvalid.Add(TEXT("output"), TEXT("OUTPUT_VAR"));
-		ReplacementsForInvalid.Add(TEXT("hvec2"), TEXT("HVEC2_VAR"));
-		ReplacementsForInvalid.Add(TEXT("hvec3"), TEXT("HVEC3_VAR"));
-		ReplacementsForInvalid.Add(TEXT("hvec4"), TEXT("HVEC4_VAR"));
-		ReplacementsForInvalid.Add(TEXT("dvec2"), TEXT("DVEC2_VAR"));
-		ReplacementsForInvalid.Add(TEXT("dvec3"), TEXT("DVEC3_VAR"));
-		ReplacementsForInvalid.Add(TEXT("dvec4"), TEXT("DVEC4_VAR"));
-		ReplacementsForInvalid.Add(TEXT("fvec2"), TEXT("FVEC2_VAR"));
-		ReplacementsForInvalid.Add(TEXT("fvec3"), TEXT("FVEC3_VAR"));
-		ReplacementsForInvalid.Add(TEXT("fvec4"), TEXT("FVEC4_VAR"));
-		ReplacementsForInvalid.Add(TEXT("sampler3DRect"), TEXT("SAMPLER3DRECT_VAR"));
-		ReplacementsForInvalid.Add(TEXT("filter"), TEXT("FILTER_VAR"));
-		ReplacementsForInvalid.Add(TEXT("image2D"), TEXT("IMAGE2D_VAR"));
-		ReplacementsForInvalid.Add(TEXT("image3D"), TEXT("IMAGE3D_VAR"));
-		ReplacementsForInvalid.Add(TEXT("imageCube"), TEXT("IMAGECUBE_VAR"));
-		ReplacementsForInvalid.Add(TEXT("iimage1D"), TEXT("IIMAGE1D_VAR"));
-		ReplacementsForInvalid.Add(TEXT("iimage2D"), TEXT("IIMAGE2D_VAR"));
-		ReplacementsForInvalid.Add(TEXT("iimage3D"), TEXT("IIMAGE3D_VAR"));
-		ReplacementsForInvalid.Add(TEXT("iimageCube"), TEXT("IIMAGECUBE_VAR"));
-		ReplacementsForInvalid.Add(TEXT("uimage1D"), TEXT("UIMAGE1D_VAR"));
-		ReplacementsForInvalid.Add(TEXT("uimage2D"), TEXT("UIMAGE2D_VAR"));
-		ReplacementsForInvalid.Add(TEXT("uimage3D"), TEXT("UIMAGE3D_VAR"));
-		ReplacementsForInvalid.Add(TEXT("uimageCube"), TEXT("UIMAGECUBE_VAR"));
-		ReplacementsForInvalid.Add(TEXT("image1DArray"), TEXT("IMAGE1DARRAY_VAR"));
-		ReplacementsForInvalid.Add(TEXT("image2DArray"), TEXT("IMAGE2DARRAY_VAR"));
-		ReplacementsForInvalid.Add(TEXT("iimage1DArray"), TEXT("IIMAGE1DARRAY_VAR"));
-		ReplacementsForInvalid.Add(TEXT("iimage2DArray"), TEXT("IIMAGE2DARRAY_VAR"));
-		ReplacementsForInvalid.Add(TEXT("uimage1DArray"), TEXT("UIMAGE1DARRAY_VAR"));
-		ReplacementsForInvalid.Add(TEXT("uimage2DArray"), TEXT("UIMAGE2DARRAY_VAR"));
-		ReplacementsForInvalid.Add(TEXT("image1DShadow"), TEXT("IMAGE1DSHADOW_VAR"));
-		ReplacementsForInvalid.Add(TEXT("image2DShadow"), TEXT("IMAGE2DSHADOW_VAR"));
-		ReplacementsForInvalid.Add(TEXT("image1DArrayShadow"), TEXT("IMAGE1DARRAYSHADOW_VAR"));
-		ReplacementsForInvalid.Add(TEXT("image2DArrayShadow"), TEXT("IMAGE2DARRAYSHADOW_VAR"));
-		ReplacementsForInvalid.Add(TEXT("imageBuffer"), TEXT("IMAGEBUFFER_VAR"));
-		ReplacementsForInvalid.Add(TEXT("iimageBuffer"), TEXT("IIMAGEBUFFER_VAR"));
-		ReplacementsForInvalid.Add(TEXT("uimageBuffer"), TEXT("UIMAGEBUFFER_VAR"));
-		ReplacementsForInvalid.Add(TEXT("sizeof"), TEXT("SIZEOF_VAR"));
-		ReplacementsForInvalid.Add(TEXT("cast"), TEXT("CAST_VAR"));
-		ReplacementsForInvalid.Add(TEXT("using"), TEXT("USING_VAR"));
-		ReplacementsForInvalid.Add(TEXT("row_major"), TEXT("ROW_MAJOR_VAR"));
-		ReplacementsForInvalid.Add(TEXT("int8_t"), TEXT("INT8_T_VAR"));
-		ReplacementsForInvalid.Add(TEXT("unsigned char"), TEXT("UNSIGNED_CHAR_VAR"));
-		ReplacementsForInvalid.Add(TEXT("uchar"), TEXT("UCHAR_VAR"));
-		ReplacementsForInvalid.Add(TEXT("uint8_t"), TEXT("UINT8_T_VAR"));
-		ReplacementsForInvalid.Add(TEXT("short"), TEXT("SHORT_VAR"));
-		ReplacementsForInvalid.Add(TEXT("int16_t"), TEXT("INT16_T_VAR"));
-		ReplacementsForInvalid.Add(TEXT("unsigned short"), TEXT("UNSIGNED_SHORT_VAR"));
-		ReplacementsForInvalid.Add(TEXT("ushort"), TEXT("USHORT_VAR"));
-		ReplacementsForInvalid.Add(TEXT("unit16_t"), TEXT("UNIT16_T_VAR"));
-		ReplacementsForInvalid.Add(TEXT("int32_t"), TEXT("INT32_T_VAR"));
-		ReplacementsForInvalid.Add(TEXT("unsigned int"), TEXT("UNSIGNED_INT_VAR"));
-		ReplacementsForInvalid.Add(TEXT("uint32_t"), TEXT("UINT32_T_VAR"));
-	}
 }
 
 FString FHlslNiagaraTranslator::GetSanitizedSymbolName(FString SymbolName, bool bCollapsNamespaces)
@@ -2467,7 +2115,9 @@ FString FHlslNiagaraTranslator::GetSanitizedSymbolName(FString SymbolName, bool
 		return SymbolName;
 	}
 
-	check(ReplacementsForInvalid.Num() != 0);
+	const UNiagaraEditorSettings* Settings = GetDefault<UNiagaraEditorSettings>();
+	check(Settings);
+	const TMap<FString, FString>& ReplacementsForInvalid = Settings->GetHLSLKeywordReplacementsMap();
 
 	FString Ret = SymbolName;
 
@@ -2484,7 +2134,7 @@ FString FHlslNiagaraTranslator::GetSanitizedSymbolName(FString SymbolName, bool
 			SplitName[i] = TEXT("INTEGER_") + SplitName[i];
 		}
 		
-		FString* FoundReplacementStr = ReplacementsForInvalid.Find(SplitName[i]); // Look for the string in the keyword protections array.
+		const FString* FoundReplacementStr = ReplacementsForInvalid.Find(SplitName[i]); // Look for the string in the keyword protections array.
 		if (FoundReplacementStr)
 		{
 			SplitName[i] = *FoundReplacementStr;
diff --git a/Engine/Plugins/FX/Niagara/Source/NiagaraEditor/Private/NiagaraHlslTranslator.h b/Engine/Plugins/FX/Niagara/Source/NiagaraEditor/Private/NiagaraHlslTranslator.h
--- a/Engine/Plugins/FX/Niagara/Source/NiagaraEditor/Private/NiagaraHlslTranslator.h
+++ b/Engine/Plugins/FX/Niagara/Source/NiagaraEditor/Private/NiagaraHlslTranslator.h
@@ -616,7 +616,4 @@ private:
 	TArray<FNiagaraVariable> InitialNamespaceVariablesMissingDefault;
 	// Variables that need to be initialized in the body or at the end of spawn.
 	TArray<FNiagaraVariable> DeferredVariablesMissingDefault;
-
-	static TMap<FString, FString> ReplacementsForInvalid;
-
 };
diff --git a/Engine/Plugins/FX/Niagara/Source/NiagaraEditor/Public/NiagaraEditorSettings.h b/Engine/Plugins/FX/Niagara/Source/NiagaraEditor/Public/NiagaraEditorSettings.h
--- a/Engine/Plugins/FX/Niagara/Source/NiagaraEditor/Public/NiagaraEditorSettings.h
+++ b/Engine/Plugins/FX/Niagara/Source/NiagaraEditor/Public/NiagaraEditorSettings.h
@@ -93,6 +93,7 @@ public:
 	/** Gets a multicast delegate which is called whenever one of the parameters in this settings object changes. */
 	static FOnNiagaraEditorSettingsChanged& OnSettingsChanged();
 
+	const TMap<FString, FString>& GetHLSLKeywordReplacementsMap()const { return HLSLKeywordReplacements; }
 protected:
 	static FOnNiagaraEditorSettingsChanged SettingsChangedDelegate;
 
@@ -115,4 +116,7 @@ private:
 
 	UPROPERTY(config)
 	TMap<FName, FNiagaraNewAssetDialogConfig> NewAssetDialogConfigMap;
+
+	UPROPERTY(config)
+	TMap<FString, FString> HLSLKeywordReplacements;
 };
diff --git a/Engine/Source/Developer/ShaderFormatVectorVM/Private/ir_vm_gen_bytecode_visitor.cpp b/Engine/Source/Developer/ShaderFormatVectorVM/Private/ir_vm_gen_bytecode_visitor.cpp
--- a/Engine/Source/Developer/ShaderFormatVectorVM/Private/ir_vm_gen_bytecode_visitor.cpp
+++ b/Engine/Source/Developer/ShaderFormatVectorVM/Private/ir_vm_gen_bytecode_visitor.cpp
@@ -545,6 +545,15 @@ void emit_short(unsigned short s, TArray<uint8>& bytecode)
 	bytecode.Add(s & 255);
 }
 
+void emit_external_func_input(variable_info_node* input, TArray<uint8>& bytecode)
+{
+	check((input->offset & VVM_EXT_FUNC_INPUT_LOC_BIT) == 0);//Ensure the offset isn't too large.
+
+	unsigned short offset = input->offset | (input->owner->location == EVectorVMOperandLocation::Constant ? 0 : VVM_EXT_FUNC_INPUT_LOC_BIT);
+	bytecode.Add(offset >> 8);
+	bytecode.Add(offset & 255);
+}
+
 struct op_base
 {
 	op_base()
@@ -571,7 +580,7 @@ struct op_base
 
 	static void validate_component_offset(_mesa_glsl_parse_state* parse_state, variable_info_node* component, unsigned op_idx)
 	{
-		if (component->offset == INDEX_NONE && component->last_read > (uint32)op_idx)
+		if (component->offset == INDEX_NONE && component->last_read >= (uint32)op_idx)
 		{
 			FString ErrorStr = FString::Printf(TEXT("Component %s of variable %s has no valid offset. Possibly uninitialized data being used."), ANSI_TO_TCHAR(component->name), ANSI_TO_TCHAR(component->owner->root->name));
 
@@ -1144,7 +1153,7 @@ struct op_external_func : public op_base
 
 		for (variable_info_node* input : inputs)
 		{
-			emit_short(input->offset, bytecode);
+			emit_external_func_input(input, bytecode);
 		}
 
 		for (variable_info_node* output : outputs)
@@ -1258,6 +1267,9 @@ class ir_gen_vvm_visitor : public ir_hierarchical_visitor
 	/** Used when building ordered_ops. Keeps track of the current operands info node so we can generate it's op code and data location. */
 	variable_info_node* curr_node;
 
+	//currently handling an output param?
+	bool is_output_param;
+
 	TArray<ir_constant*> seen_constants;
 
 	explicit ir_gen_vvm_visitor(_mesa_glsl_parse_state *in_parse_state, FVectorVMCompilationOutput& InCompilationOutput)
@@ -1270,6 +1282,7 @@ class ir_gen_vvm_visitor : public ir_hierarchical_visitor
 		, constant_table_size_bytes(0)
 		, dest_component(0)
 		, curr_node(nullptr)
+		, is_output_param(false)
 	{
 	}
 
@@ -1388,6 +1401,7 @@ class ir_gen_vvm_visitor : public ir_hierarchical_visitor
 		check(constant->type->is_scalar());//We shouldn't ever have non scalar constants.
 		check(curr_node == nullptr);
 		check(varinfo == nullptr);
+		check(!is_output_param);
 
 		const glsl_type* type = constant->type;
 		int32 same_const = seen_constants.IndexOfByPredicate(
@@ -1494,7 +1508,10 @@ class ir_gen_vvm_visitor : public ir_hierarchical_visitor
 		check(var_info);
 		check(curr_node == nullptr);
 		curr_node = var_info->root->children[index->get_int_component(0)];
-		curr_node->last_read = ordered_ops.Num() - 1;
+		if (!is_output_param)
+		{
+			curr_node->last_read = ordered_ops.Num() - 1;
+		}
 
 		return visit_continue_with_parent;
 	}
@@ -1512,8 +1529,10 @@ class ir_gen_vvm_visitor : public ir_hierarchical_visitor
 		}
 
 		curr_node = var_info->root;
-		curr_node->last_read = ordered_ops.Num() - 1;
-
+		if (!is_output_param)
+		{
+			curr_node->last_read = ordered_ops.Num() - 1;
+		}
 		return visit_continue;
 	}
 
@@ -1546,7 +1565,10 @@ class ir_gen_vvm_visitor : public ir_hierarchical_visitor
 			curr_node = curr_node->children[swiz_comp];
 		}
 
-		curr_node->last_read = ordered_ops.Num() - 1;
+		if (!is_output_param)
+		{
+			curr_node->last_read = ordered_ops.Num() - 1;
+		}
 
 		//swizzles must be the final entry in a deref chain so we have to have reached a scalar by now.
 		check(curr_node->is_scalar());
@@ -1572,7 +1594,10 @@ class ir_gen_vvm_visitor : public ir_hierarchical_visitor
 		}
 		check(prev_node != curr_node);//We have to find a child to move into.
 
-		curr_node->last_read = ordered_ops.Num() - 1;
+		if (!is_output_param)
+		{
+			curr_node->last_read = ordered_ops.Num() - 1;
+		}
 
 		return visit_continue;
 	}
@@ -1588,6 +1613,7 @@ class ir_gen_vvm_visitor : public ir_hierarchical_visitor
 	{
 		check(curr_node == nullptr);
 		check(expression->type->is_scalar());
+		check(!is_output_param)
 
 		EVectorVMBaseTypes BaseType = EVectorVMBaseTypes::Num;
 		if (expression->type->is_float())
@@ -1708,11 +1734,16 @@ class ir_gen_vvm_visitor : public ir_hierarchical_visitor
 		{
 			//Handle each param.
 			exec_node* curr = call->actual_parameters.get_head();
+			exec_node* curr_param = call->callee->parameters.get_head();
 			for (unsigned param_idx = 0; param_idx < num_operands; ++param_idx)
 			{
 				check(curr);
 				ir_rvalue* param = (ir_rvalue*)curr;
+				ir_variable* real_param = (ir_variable*)curr_param;
+
+				is_output_param = real_param->mode == ir_var_out;
 				param->accept(this);
+				is_output_param = false;
 
 				check(curr_node);
 				check(curr_node->is_scalar());
@@ -1721,6 +1752,7 @@ class ir_gen_vvm_visitor : public ir_hierarchical_visitor
 				curr_node = nullptr;
 
 				curr = curr->get_next();
+				curr_param = curr_param->get_next();
 			}
 			standard->num_operands = num_operands;
 
@@ -1928,13 +1960,15 @@ class ir_gen_vvm_visitor : public ir_hierarchical_visitor
 			foreach_iter(exec_list_iterator, iter, call->actual_parameters)
 			{
 				ir_rvalue* param = (ir_rvalue*)(iter.get());
+				check(sig_param_node);
+				ir_variable* sig_param = (ir_variable*)sig_param_node;
 
 				check(curr_node == nullptr);
+				is_output_param = sig_param->mode == ir_var_out;
 				param->accept(this);
-				check(curr_node != nullptr);				
+				check(curr_node != nullptr);	
+				is_output_param = false;
 
-				check(sig_param_node);
-				ir_variable* sig_param = (ir_variable*)sig_param_node;
 				if (sig_param->mode == ir_var_in)
 				{
 					func->inputs.Add(curr_node);
diff --git a/Engine/Source/Runtime/VectorVM/Private/VectorVM.cpp b/Engine/Source/Runtime/VectorVM/Private/VectorVM.cpp
--- a/Engine/Source/Runtime/VectorVM/Private/VectorVM.cpp
+++ b/Engine/Source/Runtime/VectorVM/Private/VectorVM.cpp
@@ -50,6 +50,103 @@ static FAutoConsoleVariableRef CVarParallelVVMChunksPerBatch(
 	ECVF_Default
 );
 
+//////////////////////////////////////////////////////////////////////////
+//  Constant Handlers
+
+struct FConstantHandlerBase
+{
+	uint16 ConstantIndex;
+	FConstantHandlerBase(FVectorVMContext& Context)
+		: ConstantIndex(VectorVM::DecodeU16(Context))
+	{}
+
+	FORCEINLINE void Advance() { }
+};
+
+template<typename T>
+struct FConstantHandler : public FConstantHandlerBase
+{
+	T Constant;
+	FConstantHandler(FVectorVMContext& Context)
+		: FConstantHandlerBase(Context)
+		, Constant(*((T*)(Context.ConstantTable + ConstantIndex)))
+	{}
+	FORCEINLINE const T& Get() { return Constant; }
+	FORCEINLINE const T& GetAndAdvance() { return Constant; }
+};
+
+struct FDataSetOffsetHandler : FConstantHandlerBase
+{
+	uint32 Offset;
+	FDataSetOffsetHandler(FVectorVMContext& Context)
+		: FConstantHandlerBase(Context)
+		, Offset(Context.DataSetOffsetTable[ConstantIndex])
+	{}
+	FORCEINLINE const uint32 Get() { return Offset; }
+	FORCEINLINE const uint32 GetAndAdvance() { return Offset; }
+};
+
+template<>
+struct FConstantHandler<VectorRegister> : public FConstantHandlerBase
+{
+	VectorRegister Constant;
+	FConstantHandler(FVectorVMContext& Context)
+		: FConstantHandlerBase(Context)
+		, Constant(VectorLoadFloat1(&Context.ConstantTable[ConstantIndex]))
+	{}
+	FORCEINLINE const VectorRegister Get() { return Constant; }
+	FORCEINLINE const VectorRegister GetAndAdvance() { return Constant; }
+};
+
+template<>
+struct FConstantHandler<VectorRegisterInt> : public FConstantHandlerBase
+{
+	VectorRegisterInt Constant;
+	FConstantHandler(FVectorVMContext& Context)
+		: FConstantHandlerBase(Context)
+		, Constant(VectorIntLoad1(&Context.ConstantTable[ConstantIndex]))
+	{}
+	FORCEINLINE const VectorRegisterInt Get() { return Constant; }
+	FORCEINLINE const VectorRegisterInt GetAndAdvance() { return Constant; }
+};
+
+
+//////////////////////////////////////////////////////////////////////////
+// Register handlers.
+// Handle reading of a register, advancing the pointer with each read.
+
+struct FRegisterHandlerBase
+{
+	int32 RegisterIndex;
+	FORCEINLINE FRegisterHandlerBase(FVectorVMContext& Context)
+		: RegisterIndex(VectorVM::DecodeU16(Context))
+	{}
+
+};
+
+template<typename T>
+struct FRegisterHandler : public FRegisterHandlerBase
+{
+private:
+	T * RESTRICT Register;
+public:
+	FORCEINLINE FRegisterHandler(FVectorVMContext& Context)
+		: FRegisterHandlerBase(Context)
+		, Register((T*)Context.RegisterTable[RegisterIndex])
+	{}
+	FORCEINLINE const T Get() { return *Register; }
+	FORCEINLINE T* GetDest() { return Register; }
+	FORCEINLINE void Advance() { ++Register; }
+	FORCEINLINE const T GetAndAdvance()
+	{
+		return *Register++;
+	}
+	FORCEINLINE T* GetDestAndAdvance()
+	{
+		return Register++;
+	}
+};
+
 //////////////////////////////////////////////////////////////////////////
 
 FVectorVMContext::FVectorVMContext()
@@ -211,7 +308,7 @@ struct TUnaryKernel
 {
 	static void Exec(FVectorVMContext& Context)
 	{
-		uint32 SrcOpTypes = DecodeSrcOperandTypes(Context);
+		uint32 SrcOpTypes = VectorVM::DecodeSrcOperandTypes(Context);
 		switch (SrcOpTypes)
 		{
 		case SRCOP_RRR: TUnaryKernelHandler<Kernel, DstHandler, RegisterHandler, NumInstancesPerOp>::Exec(Context); break;
@@ -235,7 +332,7 @@ struct TBinaryKernel
 {
 	static void Exec(FVectorVMContext& Context)
 	{
-		uint32 SrcOpTypes = DecodeSrcOperandTypes(Context);
+		uint32 SrcOpTypes = VectorVM::DecodeSrcOperandTypes(Context);
 		switch (SrcOpTypes)
 		{
 		case SRCOP_RRR: TBinaryKernelHandler<Kernel, DstHandler, RegisterHandler, RegisterHandler, NumInstancesPerOp>::Exec(Context); break;
@@ -259,7 +356,7 @@ struct TTrinaryKernel
 {
 	static void Exec(FVectorVMContext& Context)
 	{
-		uint32 SrcOpTypes = DecodeSrcOperandTypes(Context);
+		uint32 SrcOpTypes = VectorVM::DecodeSrcOperandTypes(Context);
 		switch (SrcOpTypes)
 		{
 		case SRCOP_RRR: TTrinaryKernelHandler<Kernel, DstHandler, RegisterHandler, RegisterHandler, RegisterHandler, NumInstancesPerOp>::Exec(Context); break;
@@ -884,15 +981,15 @@ struct FScalarKernelAcquireID
 {
 	static VM_FORCEINLINE void Exec(FVectorVMContext& Context)
 	{
-		int32 DataSetIndex = DecodeU16(Context);
+		int32 DataSetIndex = VectorVM::DecodeU16(Context);
 		TArray<int32>&RESTRICT FreeIDTable = *Context.DataSetMetaTable[DataSetIndex].FreeIDTable;
 
 		int32 Tag = Context.DataSetMetaTable[DataSetIndex].IDAcquireTag;
 
-		int32 IDIndexReg = DecodeU16(Context);
+		int32 IDIndexReg = VectorVM::DecodeU16(Context);
 		int32*RESTRICT IDIndex = (int32*)(Context.RegisterTable[IDIndexReg]);
 
-		int32 IDTagReg = DecodeU16(Context);
+		int32 IDTagReg = VectorVM::DecodeU16(Context);
 		int32*RESTRICT IDTag = (int32*)(Context.RegisterTable[IDTagReg]);
 
 		int32* NumFreeIDs = Context.DataSetMetaTable[DataSetIndex].NumFreeIDs;
@@ -916,10 +1013,10 @@ struct FScalarKernelUpdateID
 {
 	static VM_FORCEINLINE void Exec(FVectorVMContext& Context)
 	{
-		int32 DataSetIndex = DecodeU16(Context);
-		int32 InstanceIDRegisterIndex = DecodeU16(Context);
-		int32 InstanceIndexRegisterIndex = DecodeU16(Context);
-		
+		int32 DataSetIndex = VectorVM::DecodeU16(Context);
+		int32 InstanceIDRegisterIndex = VectorVM::DecodeU16(Context);
+		int32 InstanceIndexRegisterIndex = VectorVM::DecodeU16(Context);
+
 		TArray<int32>&RESTRICT IDTable = *Context.DataSetMetaTable[DataSetIndex].IDTable;
 		TArray<int32>&RESTRICT FreeIDTable = *Context.DataSetMetaTable[DataSetIndex].FreeIDTable;
 		int32 InstanceOffset = Context.DataSetMetaTable[DataSetIndex].InstanceOffset;// +Context.StartInstance;
@@ -975,9 +1072,9 @@ struct FVectorKernelReadInput
 	{
 		static const int32 InstancesPerVector = sizeof(VectorRegister) / sizeof(T);
 
-		int32 DataSetIndex = DecodeU16(Context);
-		int32 InputRegisterIdx = DecodeU16(Context);
-		int32 DestRegisterIdx = DecodeU16(Context);
+		int32 DataSetIndex = VectorVM::DecodeU16(Context);
+		int32 InputRegisterIdx = VectorVM::DecodeU16(Context);
+		int32 DestRegisterIdx = VectorVM::DecodeU16(Context);
 		int32 Loops = Align(Context.NumInstances, InstancesPerVector) / InstancesPerVector;
 
 		VectorRegister* DestReg = (VectorRegister*)(Context.RegisterTable[DestRegisterIdx]);
@@ -1007,9 +1104,9 @@ struct FVectorKernelReadInputNoAdvance
 	{
 		static const int32 InstancesPerVector = sizeof(VectorRegister) / sizeof(T);
 
-		int32 DataSetIndex = DecodeU16(Context);
-		int32 InputRegisterIdx = DecodeU16(Context);
-		int32 DestRegisterIdx = DecodeU16(Context);
+		int32 DataSetIndex = VectorVM::DecodeU16(Context);
+		int32 InputRegisterIdx = VectorVM::DecodeU16(Context);
+		int32 DestRegisterIdx = VectorVM::DecodeU16(Context);
 		int32 Loops = Align(Context.NumInstances, InstancesPerVector) / InstancesPerVector;
 
 		VectorRegister* DestReg = (VectorRegister*)(Context.RegisterTable[DestRegisterIdx]);
@@ -1074,7 +1171,7 @@ struct FScalarKernelWriteOutputIndexed
 {
 	static VM_FORCEINLINE void Exec(FVectorVMContext& Context)
 	{
-		uint32 SrcOpTypes = DecodeSrcOperandTypes(Context);		
+		uint32 SrcOpTypes = VectorVM::DecodeSrcOperandTypes(Context);
 		switch (SrcOpTypes)
 		{
 		case SRCOP_RRR: TTrinaryOutputKernelHandler<FScalarKernelWriteOutputIndexed, FOutputRegisterHandler<T>, FDataSetOffsetHandler, FRegisterHandler<int32>, FRegisterHandler<T>, 1>::Exec(Context); break;
@@ -1096,7 +1193,7 @@ struct FDataSetCounterHandler
 {
 	int32* Counter;
 	FDataSetCounterHandler(FVectorVMContext& Context)
-		: Counter(Context.DataSetIndexTable + DecodeU16(Context))
+		: Counter(Context.DataSetIndexTable + VectorVM::DecodeU16(Context))
 	{}
 
 	VM_FORCEINLINE void Advance() { }
@@ -1110,7 +1207,7 @@ struct FScalarKernelAcquireCounterIndex
 {
 	static VM_FORCEINLINE void Exec(FVectorVMContext& Context)
 	{
-		uint32 SrcOpTypes = DecodeSrcOperandTypes(Context);
+		uint32 SrcOpTypes = VectorVM::DecodeSrcOperandTypes(Context);
 		switch (SrcOpTypes)
 		{
 		case SRCOP_RRR: TBinaryKernelHandler<FScalarKernelAcquireCounterIndex, FRegisterHandler<int32>, FDataSetCounterHandler, FRegisterHandler<int32>, 1>::Exec(Context); break;
@@ -1165,7 +1262,7 @@ struct FKernelExternalFunctionCall
 {
 	static void Exec(FVectorVMContext& Context)
 	{
-		uint32 ExternalFuncIdx = DecodeU8(Context);
+		uint32 ExternalFuncIdx = VectorVM::DecodeU8(Context);
 		Context.ExternalFunctionTable[ExternalFuncIdx].Execute(Context);
 	}
 };
diff --git a/Engine/Source/Runtime/VectorVM/Public/VectorVM.h b/Engine/Source/Runtime/VectorVM/Public/VectorVM.h
--- a/Engine/Source/Runtime/VectorVM/Public/VectorVM.h
+++ b/Engine/Source/Runtime/VectorVM/Public/VectorVM.h
@@ -197,43 +197,11 @@ namespace VectorVM
 		FirstOutputRegister = FirstInputRegister + MaxInputRegisters,
 		MaxRegisters = NumTempRegisters + MaxInputRegisters + MaxOutputRegisters + MaxConstants,
 	};
+}
 
-	/** Get total number of op-codes */
-	VECTORVM_API uint8 GetNumOpCodes();
-
-#if WITH_EDITOR
-	VECTORVM_API FString GetOpName(EVectorVMOp Op);
-	VECTORVM_API FString GetOperandLocationName(EVectorVMOperandLocation Location);
-#endif
-
-	VECTORVM_API uint8 CreateSrcOperandMask(EVectorVMOperandLocation Type0, EVectorVMOperandLocation Type1 = EVectorVMOperandLocation::Register, EVectorVMOperandLocation Type2 = EVectorVMOperandLocation::Register);
-
-	/**
-	 * Execute VectorVM bytecode.
-	 */
-	VECTORVM_API void Exec(
-		uint8 const* Code,
-		uint8** InputRegisters,
-		int32 NumInputRegisters,
-		uint8** OutputRegisters,
-		int32 NumOutputRegisters,
-		uint8 const* ConstantTable,
-		TArray<FDataSetMeta> &DataSetMetaTable,
-		FVMExternalFunction* ExternalFunctionTable,
-		void** UserPtrTable,
-		int32 NumInstances
-#if STATS
-		, const TArray<TStatId>& StatScopes
-#endif
-		);
-
-	VECTORVM_API void Init();
-} // namespace VectorVM
-
-
-  /**
-  * Context information passed around during VM execution.
-  */
+/**
+* Context information passed around during VM execution.
+*/
 struct FVectorVMContext : TThreadSingleton<FVectorVMContext>
 {
 	/** Pointer to the next element in the byte code. */
@@ -293,153 +261,169 @@ struct FVectorVMContext : TThreadSingleton<FVectorVMContext>
 	}
 };
 
-FORCEINLINE uint8 DecodeU8(FVectorVMContext& Context)
-{
-	return *Context.Code++;
-}
-
-FORCEINLINE uint16 DecodeU16(FVectorVMContext& Context)
-{
-	return ((uint16)DecodeU8(Context) << 8) + DecodeU8(Context);
-}
-
-FORCEINLINE uint32 DecodeU32(FVectorVMContext& Context)
-{
-	return ((uint32)DecodeU8(Context) << 24) + (uint32)(DecodeU8(Context) << 16) + (uint32)(DecodeU8(Context) << 8) + DecodeU8(Context);
-}
-
-/** Decode the next operation contained in the bytecode. */
-FORCEINLINE EVectorVMOp DecodeOp(FVectorVMContext& Context)
-{
-	return static_cast<EVectorVMOp>(DecodeU8(Context));
-}
-
-FORCEINLINE uint8 DecodeSrcOperandTypes(FVectorVMContext& Context)
-{
-	return DecodeU8(Context);
-}
-
-//////////////////////////////////////////////////////////////////////////
-/** Constant handler. */
-
-struct FConstantHandlerBase
-{
-	uint16 ConstantIndex;
-	FConstantHandlerBase(FVectorVMContext& Context)
-		: ConstantIndex(DecodeU16(Context))
-	{}
-
-	FORCEINLINE void Advance() { }
-};
-
-template<typename T>
-struct FConstantHandler : public FConstantHandlerBase
+namespace VectorVM
 {
-	T Constant;
-	FConstantHandler(FVectorVMContext& Context)
-		: FConstantHandlerBase(Context)
-		, Constant(*((T*)(Context.ConstantTable + ConstantIndex)))
-	{}
-	FORCEINLINE const T& Get() { return Constant; }
-	FORCEINLINE const T& GetAndAdvance() { return Constant; }
-};
+	/** Get total number of op-codes */
+	VECTORVM_API uint8 GetNumOpCodes();
 
+#if WITH_EDITOR
+	VECTORVM_API FString GetOpName(EVectorVMOp Op);
+	VECTORVM_API FString GetOperandLocationName(EVectorVMOperandLocation Location);
+#endif
 
+	VECTORVM_API uint8 CreateSrcOperandMask(EVectorVMOperandLocation Type0, EVectorVMOperandLocation Type1 = EVectorVMOperandLocation::Register, EVectorVMOperandLocation Type2 = EVectorVMOperandLocation::Register);
 
-struct FDataSetOffsetHandler : FConstantHandlerBase
-{
-	uint32 Offset;
-	FDataSetOffsetHandler(FVectorVMContext& Context)
-		: FConstantHandlerBase(Context)
-		, Offset(Context.DataSetOffsetTable[ConstantIndex])
-	{}
-	FORCEINLINE const uint32 Get() { return Offset; }
-	FORCEINLINE const uint32 GetAndAdvance() { return Offset; }
-};
+	/**
+	 * Execute VectorVM bytecode.
+	 */
+	VECTORVM_API void Exec(
+		uint8 const* Code,
+		uint8** InputRegisters,
+		int32 NumInputRegisters,
+		uint8** OutputRegisters,
+		int32 NumOutputRegisters,
+		uint8 const* ConstantTable,
+		TArray<FDataSetMeta> &DataSetMetaTable,
+		FVMExternalFunction* ExternalFunctionTable,
+		void** UserPtrTable,
+		int32 NumInstances
+#if STATS
+		, const TArray<TStatId>& StatScopes
+#endif
+	);
 
+	VECTORVM_API void Init();
 
+	FORCEINLINE uint8 DecodeU8(FVectorVMContext& Context)
+	{
+		return *Context.Code++;
+	}
 
-template<>
-struct FConstantHandler<VectorRegister> : public FConstantHandlerBase
-{
-	VectorRegister Constant;
-	FConstantHandler(FVectorVMContext& Context)
-		: FConstantHandlerBase(Context)
-		, Constant(VectorLoadFloat1(&Context.ConstantTable[ConstantIndex]))
-	{}
-	FORCEINLINE const VectorRegister Get() { return Constant; }
-	FORCEINLINE const VectorRegister GetAndAdvance() { return Constant; }
-};
+	FORCEINLINE uint16 DecodeU16(FVectorVMContext& Context)
+	{
+		return ((uint16)DecodeU8(Context) << 8) + DecodeU8(Context);
+	}
 
-template<>
-struct FConstantHandler<VectorRegisterInt> : public FConstantHandlerBase
-{
-	VectorRegisterInt Constant;
-	FConstantHandler(FVectorVMContext& Context)
-		: FConstantHandlerBase(Context)
-		, Constant(VectorIntLoad1(&Context.ConstantTable[ConstantIndex]))
-	{}
-	FORCEINLINE const VectorRegisterInt Get() { return Constant; }
-	FORCEINLINE const VectorRegisterInt GetAndAdvance() { return Constant; }
-};
+	FORCEINLINE uint32 DecodeU32(FVectorVMContext& Context)
+	{
+		return ((uint32)DecodeU8(Context) << 24) + (uint32)(DecodeU8(Context) << 16) + (uint32)(DecodeU8(Context) << 8) + DecodeU8(Context);
+	}
 
-//////////////////////////////////////////////////////////////////////////
-// Register handlers.
-// Handle reading of a register, advancing the pointer with each read.
+	/** Decode the next operation contained in the bytecode. */
+	FORCEINLINE EVectorVMOp DecodeOp(FVectorVMContext& Context)
+	{
+		return static_cast<EVectorVMOp>(DecodeU8(Context));
+	}
 
-struct FRegisterHandlerBase
-{
-	int32 RegisterIndex;
-	FORCEINLINE FRegisterHandlerBase(FVectorVMContext& Context)
-		: RegisterIndex(DecodeU16(Context))
-	{}
+	FORCEINLINE uint8 DecodeSrcOperandTypes(FVectorVMContext& Context)
+	{
+		return DecodeU8(Context);
+	}
 
-	FORCEINLINE bool IsValid() const { return RegisterIndex != 0xFFFF; }
-};
+#define VVM_EXT_FUNC_INPUT_LOC_BIT (unsigned short)(1<<15)
+#define VVM_EXT_FUNC_INPUT_LOC_MASK (unsigned short)~VVM_EXT_FUNC_INPUT_LOC_BIT
 
-template<typename T>
-struct FUserPtrHandler
-{
-	int32 UserPtrIdx;
-	T* Ptr;
-	FUserPtrHandler(FVectorVMContext& Context)
-		: UserPtrIdx(*(int32*)(Context.ConstantTable + DecodeU16(Context)))
-		, Ptr((T*)Context.UserPtrTable[UserPtrIdx])
+	template<typename T>
+	struct FUserPtrHandler
 	{
-		check(UserPtrIdx != INDEX_NONE);
-	}
-	FORCEINLINE T* Get() { return Ptr; }
-	FORCEINLINE T* operator->() { return Ptr; }
-	FORCEINLINE operator T*() { return Ptr; }
-};
+		int32 UserPtrIdx;
+		T* Ptr;
+		FUserPtrHandler(FVectorVMContext& Context)
+			: UserPtrIdx(*(int32*)(Context.ConstantTable + (DecodeU16(Context))))
+			, Ptr((T*)Context.UserPtrTable[UserPtrIdx])
+		{
+			check(UserPtrIdx != INDEX_NONE);
+		}
+		FORCEINLINE T* Get() { return Ptr; }
+		FORCEINLINE T* operator->() { return Ptr; }
+		FORCEINLINE operator T*() { return Ptr; }
+	};
 
-template<typename T>
-struct FRegisterHandler : public FRegisterHandlerBase
-{
-private:
-	T Dummy;
-	T* RESTRICT Register;
-	uint32 AdvanceOffset;
-public:
-	FORCEINLINE FRegisterHandler(FVectorVMContext& Context)
-		: FRegisterHandlerBase(Context)
-		, Register(IsValid() ? (T*)Context.RegisterTable[RegisterIndex] : &Dummy)
-		, AdvanceOffset(IsValid() ? 1 : 0)
-	{}
-	FORCEINLINE const T Get() { return *Register; }
-	FORCEINLINE T* GetDest() { return Register; }
-	FORCEINLINE void Advance() { Register += AdvanceOffset; }
-	FORCEINLINE const T GetAndAdvance()
+	// A flexible handler that can deal with either constant or register inputs.
+	template<typename T>
+	struct FExternalFuncInputHandler
 	{
-		T* Ret = Register;
-		Register += AdvanceOffset;
-		return *Ret;
-	}
-	FORCEINLINE T* GetDestAndAdvance()
+	private:
+		/** Either byte offset into constant table or offset into register table deepening on VVM_INPUT_LOCATION_BIT */
+		uint32 InputOffset;
+		T* RESTRICT InputPtr;
+		uint32 AdvanceOffset;
+
+	public:
+		FORCEINLINE FExternalFuncInputHandler(FVectorVMContext& Context)
+			: InputOffset(DecodeU16(Context))
+			, InputPtr(IsConstant() ? (T*)(Context.ConstantTable + GetOffset()) : (T*)Context.RegisterTable[GetOffset()])
+			, AdvanceOffset(IsConstant() ? 0 : 1)
+		{}
+
+		FORCEINLINE bool IsConstant()const { return !IsRegister(); }
+		FORCEINLINE bool IsRegister()const { return InputOffset & VVM_EXT_FUNC_INPUT_LOC_BIT; }
+		FORCEINLINE int32 GetOffset()const { return InputOffset & VVM_EXT_FUNC_INPUT_LOC_MASK; }
+
+		FORCEINLINE const T Get() { return *InputPtr; }
+		FORCEINLINE T* GetDest() { return InputPtr; }
+		FORCEINLINE void Advance() { InputPtr += AdvanceOffset; }
+		FORCEINLINE const T GetAndAdvance()
+		{
+			T* Ret = InputPtr;
+			InputPtr += AdvanceOffset;
+			return *Ret;
+		}
+		FORCEINLINE T* GetDestAndAdvance()
+		{
+			T* Ret = InputPtr;
+			InputPtr += AdvanceOffset;
+			return Ret;
+		}
+	};
+	
+	template<typename T>
+	struct FExternalFuncRegisterHandler
 	{
-		T* Ret = Register;
-		Register += AdvanceOffset;
-		return Ret;
-	}
-};
+	private:
+		uint32 RegisterIndex;
+		uint32 AdvanceOffset;
+		T Dummy;
+		T* RESTRICT Register;
+	public:
+		FORCEINLINE FExternalFuncRegisterHandler(FVectorVMContext& Context)
+			: RegisterIndex(DecodeU16(Context) & VVM_EXT_FUNC_INPUT_LOC_MASK)
+			, AdvanceOffset(IsValid() ? 1 : 0)
+			, Register(IsValid() ? (T*)Context.RegisterTable[RegisterIndex] : &Dummy)
+		{}
+		FORCEINLINE bool IsValid() const { return RegisterIndex != (uint16)VVM_EXT_FUNC_INPUT_LOC_MASK; }
+
+		FORCEINLINE const T Get() { return *Register; }
+		FORCEINLINE T* GetDest() { return Register; }
+		FORCEINLINE void Advance() { Register += AdvanceOffset; }
+		FORCEINLINE const T GetAndAdvance()
+		{
+			T* Ret = Register;
+			Register += AdvanceOffset;
+			return *Ret;
+		}
+		FORCEINLINE T* GetDestAndAdvance()
+		{
+			T* Ret = Register;
+			Register += AdvanceOffset;
+			return Ret;
+		}
+	};
+
+	template<typename T>
+	struct FExternalFuncConstHandler
+	{
+		uint16 ConstantIndex;
+		T Constant;
+		FExternalFuncConstHandler(FVectorVMContext& Context)
+			: ConstantIndex(VectorVM::DecodeU16(Context) & VVM_EXT_FUNC_INPUT_LOC_MASK)
+			, Constant(*((T*)(Context.ConstantTable + ConstantIndex)))
+		{}
+		FORCEINLINE const T& Get() { return Constant; }
+		FORCEINLINE const T& GetAndAdvance() { return Constant; }
+		FORCEINLINE void Advance() { }
+	};
+} // namespace VectorVM
+
+
 
