Fix HoudiniNiagara GetLastPointIDToSpawnAtTime() becoming stuck if all points would spawn at time t=0.

diff --git a/Engine/Plugins/FX/HoudiniNiagara/Source/HoudiniNiagara/Private/HoudiniCSV.cpp b/Engine/Plugins/FX/HoudiniNiagara/Source/HoudiniNiagara/Private/HoudiniCSV.cpp
--- a/Engine/Plugins/FX/HoudiniNiagara/Source/HoudiniNiagara/Private/HoudiniCSV.cpp
+++ b/Engine/Plugins/FX/HoudiniNiagara/Source/HoudiniNiagara/Private/HoudiniCSV.cpp
@@ -846,49 +846,56 @@ bool UHoudiniCSV::GetLastRowIndexAtTime(const float& desiredTime, int32& lastRow
 // If the CSV file doesn't have time informations, returns false and set the LastIndex to the last point
 // If desiredTime is smaller than the first point time, LastIndex will be set to -1
 // If desiredTime is higher than the last point time in the csv file, LastIndex will be set to the last point's index
-bool UHoudiniCSV::GetLastPointIDToSpawnAtTime( const float& desiredTime, int32& lastID ) const
+bool UHoudiniCSV::GetLastPointIDToSpawnAtTime(const float& desiredTime, int32& lastID) const
 {
 	/*
 	// If we dont have proper time info, always return the last point
-	int32 TimeColumnIndex = GetAttributeColumnIndex(EHoudiniAttributes::TIME);    
-    if ( TimeColumnIndex < 0 || TimeColumnIndex >= NumberOfColumns )
-    {
+	int32 TimeColumnIndex = GetAttributeColumnIndex(EHoudiniAttributes::TIME);
+	if ( TimeColumnIndex < 0 || TimeColumnIndex >= NumberOfColumns )
+	{
 		lastID = NumberOfPoints - 1;
 		return false;
-    }
+	}
 	*/
 
-    float temp_time = 0.0f;
-	if ( !SpawnTimes.IsValidIndex( NumberOfPoints - 1 ) )
+	float temp_time = 0.0f;
+	if (!SpawnTimes.IsValidIndex(NumberOfPoints - 1))
 	{
 		lastID = NumberOfPoints - 1;
+		LastDesiredTime = desiredTime;
 		return false;
 	}
 
-	if ( SpawnTimes[ NumberOfPoints - 1 ] < desiredTime )
+	else if (SpawnTimes[NumberOfPoints - 1] < desiredTime && desiredTime > LastDesiredTime)
 	{
-		// We didn't find a suitable index because the desired time is higher than our last time value
+		// We didn't find a suitable index because the desired time is higher than the last index's time value and the spawn time has not looped over
 		lastID = NumberOfPoints - 1;
+		LastDesiredTime = desiredTime;
 		return true;
 	}
 
-	// Iterates through all the points
-	lastID = INDEX_NONE;
-	for ( int32 n = 0; n < NumberOfPoints; n++ )
+	else
 	{
-		temp_time = SpawnTimes[ n ];
-
-		if ( temp_time == desiredTime )
+		// Iterates through all the points
+		lastID = GetNumberOfPoints();
+		for (int32 n = 0; n < NumberOfPoints; n++)
 		{
-			lastID = n;
-		}
-		else if ( temp_time > desiredTime )
-		{
-			lastID = n - 1;
-			return true;
+			temp_time = SpawnTimes[n];
+
+			if (temp_time == desiredTime)
+			{
+				lastID = n;
+			}
+			else if (temp_time > desiredTime)
+			{
+				lastID = n - 1;
+				LastDesiredTime = desiredTime;
+				return true;
+			}
 		}
 	}
 
+	LastDesiredTime = desiredTime;
 	return true;
 }
 
diff --git a/Engine/Plugins/FX/HoudiniNiagara/Source/HoudiniNiagara/Public/HoudiniCSV.h b/Engine/Plugins/FX/HoudiniNiagara/Source/HoudiniNiagara/Public/HoudiniCSV.h
--- a/Engine/Plugins/FX/HoudiniNiagara/Source/HoudiniNiagara/Public/HoudiniCSV.h
+++ b/Engine/Plugins/FX/HoudiniNiagara/Source/HoudiniNiagara/Public/HoudiniCSV.h
@@ -270,4 +270,7 @@ class HOUDININIAGARA_API UHoudiniCSV : public UObject
 
 	UPROPERTY()
 	bool UseCustomTitleRow;
+
+	// Float to track last desired time of GetLastPointIDToSpawnAtTime()
+	mutable float LastDesiredTime;
 };
\ No newline at end of file
