Linux: whitelist HoudiniNiagara.

diff --git a/Engine/Plugins/FX/HoudiniNiagara/HoudiniNiagara.uplugin b/Engine/Plugins/FX/HoudiniNiagara/HoudiniNiagara.uplugin
--- a/Engine/Plugins/FX/HoudiniNiagara/HoudiniNiagara.uplugin
+++ b/Engine/Plugins/FX/HoudiniNiagara/HoudiniNiagara.uplugin
@@ -19,13 +19,13 @@
 			"Name": "HoudiniNiagara",
 			"Type": "Runtime",
 			"LoadingPhase": "Default",
-			"WhitelistPlatforms": [ "Win64", "Win32", "PS4", "XboxOne", "Mac" ]
+			"WhitelistPlatforms": [ "Win64", "Win32", "PS4", "XboxOne", "Mac", "Linux" ]
 		},
 		{
 			"Name": "HoudiniNiagaraEditor",
 			"Type": "Editor",
 			"LoadingPhase": "Default",
-			"WhitelistPlatforms": [ "Win64", "Win32", "PS4", "XboxOne", "Mac" ]
+			"WhitelistPlatforms": [ "Win64", "Win32", "PS4", "XboxOne", "Mac", "Linux" ]
 		}
 	],
 
@@ -35,4 +35,4 @@
 			"Enabled": true
 		}
 	]
-}
\ No newline at end of file
+}
